[StormSyncMotionDesignBridge]
Move plugin

diff --git a/Engine/Plugins/Experimental/StormSyncAvalancheBridge/Source/StormSyncAvaBridge/Private/ModularFeature/StormSyncAvaSyncProvider.cpp b/Engine/Plugins/Experimental/StormSyncAvalancheBridge/Source/StormSyncAvaBridge/Private/ModularFeature/StormSyncAvaSyncProvider.cpp
new file mode 100644
--- /dev/null
+++ b/Engine/Plugins/Experimental/StormSyncAvalancheBridge/Source/StormSyncAvaBridge/Private/ModularFeature/StormSyncAvaSyncProvider.cpp
@@ -0,0 +1,370 @@
+// Copyright Epic Games, Inc. All Rights Reserved.
+
+#include "StormSyncAvaSyncProvider.h"
+
+#include "IAvaMediaModule.h"
+#include "IStormSyncTransportClientModule.h"
+#include "Playback/AvaMediaPlaybackServer.h"
+#include "Playback/IAvaMediaPlaybackClient.h"
+#include "StormSyncAvaBridgeCommon.h"
+#include "StormSyncAvaBridgeLog.h"
+#include "StormSyncCoreDelegates.h"
+
+#define LOCTEXT_NAMESPACE "StormSyncAvaSyncProvider"
+
+FStormSyncAvaSyncProvider::FStormSyncAvaSyncProvider()
+{
+	FStormSyncCoreDelegates::OnPakAssetExtracted.AddRaw(this, &FStormSyncAvaSyncProvider::OnPakAssetExtracted);
+}
+
+FStormSyncAvaSyncProvider::~FStormSyncAvaSyncProvider()
+{
+	FStormSyncCoreDelegates::OnPakAssetExtracted.RemoveAll(this);
+}
+
+FName FStormSyncAvaSyncProvider::GetName() const
+{
+	return TEXT("StormSyncAvaSyncProvider");
+}
+
+void FStormSyncAvaSyncProvider::SyncToAll(const TArray<FName>& InPackageNames)
+{
+	STORM_SYNC_AVA_LOG(Display, TEXT("FStormSyncAvaSyncProvider::SyncToAll InPackageNames: %d"), InPackageNames.Num());
+
+	const FStormSyncPackageDescriptor PackageDescriptor;
+	IStormSyncTransportClientModule::Get().SynchronizePackages(PackageDescriptor, InPackageNames);
+}
+
+void FStormSyncAvaSyncProvider::PushToRemote(const FString& InRemoteName, const TArray<FName>& InPackageNames, const FOnAvaSyncResponse& DoneDelegate)
+{
+	STORM_SYNC_AVA_LOG(Display, TEXT("FStormSyncAvaSyncProvider::PushToRemote InRemoteName: %s, InPackageNames: %d"), *InRemoteName, InPackageNames.Num());
+	
+	FText ErrorText;
+	FMessageAddress MessageAddress;
+	// Get address from either client or server user data, depending on local client / server state
+	if (!GetAddressFromUserData(InRemoteName, UE::StormSync::AvaBridgeCommon::StormSyncClientAddressKey, MessageAddress, &ErrorText))
+	{
+		STORM_SYNC_AVA_LOG(Error, TEXT("FStormSyncAvaSyncProvider::PushToRemote - %s"), *ErrorText.ToString());
+		DoneDelegate.ExecuteIfBound(CreateErrorResponse(ErrorText));
+		return;
+	}
+
+	const FOnStormSyncPushComplete Delegate = FOnStormSyncPushComplete::CreateLambda([DoneDelegate](const TSharedPtr<FStormSyncTransportPushResponse>& Response)
+	{
+		if (!Response.IsValid())
+		{
+			DoneDelegate.ExecuteIfBound(CreateErrorResponse(LOCTEXT("Error_PushInvalidResponse", "Received invalid shared ptr response")));
+			return;
+		}
+		
+		STORM_SYNC_AVA_LOG(Display, TEXT("FStormSyncAvaSyncProvider::PushToRemote - Response: %s"), *Response->ToString());
+
+		if (DoneDelegate.IsBound())
+		{
+			FAvaSyncResponse SyncResponse = ConvertSyncResponse(Response);
+			DoneDelegate.Execute(MakeShared<FAvaSyncResponse>(SyncResponse));
+		}
+	});
+
+	const FStormSyncPackageDescriptor PackageDescriptor;
+	IStormSyncTransportClientModule::Get().PushPackages(PackageDescriptor, InPackageNames, MessageAddress, Delegate);
+}
+
+void FStormSyncAvaSyncProvider::PullFromRemote(const FString& InRemoteName, const TArray<FName>& InPackageNames, const FOnAvaSyncResponse& DoneDelegate)
+{
+	STORM_SYNC_AVA_LOG(Display, TEXT("FStormSyncAvaSyncProvider::PullFromRemote InRemoteName: %s, InPackageNames: %d"), *InRemoteName, InPackageNames.Num());
+
+	FText ErrorText;
+	FMessageAddress MessageAddress;
+	// Get address from either client or server user data, depending on local client / server state
+	if (!GetAddressFromUserData(InRemoteName, UE::StormSync::AvaBridgeCommon::StormSyncClientAddressKey, MessageAddress, &ErrorText))
+	{
+		STORM_SYNC_AVA_LOG(Error, TEXT("FStormSyncAvaSyncProvider::PullFromRemote - %s"), *ErrorText.ToString());
+		DoneDelegate.ExecuteIfBound(CreateErrorResponse(ErrorText));
+		return;
+	}
+
+	const FOnStormSyncPullComplete Delegate = FOnStormSyncPullComplete::CreateLambda([DoneDelegate](const TSharedPtr<FStormSyncTransportPullResponse>& Response)
+	{
+		if (!Response.IsValid())
+		{
+			DoneDelegate.ExecuteIfBound(CreateErrorResponse(LOCTEXT("Error_PullInvalidResponse", "Received invalid shared ptr response")));
+			return;
+		}
+		
+		STORM_SYNC_AVA_LOG(Display, TEXT("FStormSyncAvaSyncProvider::PullFromRemote - Response: %s"), *Response->ToString());
+
+		if (DoneDelegate.IsBound())
+		{
+			FAvaSyncResponse SyncResponse = ConvertSyncResponse(Response);
+			DoneDelegate.Execute(MakeShared<FAvaSyncResponse>(SyncResponse));
+		}
+	});
+
+	const FStormSyncPackageDescriptor PackageDescriptor;
+	IStormSyncTransportClientModule::Get().PullPackages(PackageDescriptor, InPackageNames, MessageAddress, Delegate);
+}
+
+void FStormSyncAvaSyncProvider::CompareWithRemote(const FString& InRemoteName, const TArray<FName>& InPackageNames, const FOnAvaSyncCompareResponse& DoneDelegate)
+{
+	STORM_SYNC_AVA_LOG(Display, TEXT("FStormSyncAvaSyncProvider::CompareWithRemote InRemoteName: %s, InPackageNames: %d"), *InRemoteName, InPackageNames.Num());
+
+	const bool bIsClient = IsPlaybackClient();
+	const bool bIsServer = IsPlaybackServer();
+	STORM_SYNC_AVA_LOG(Display, TEXT("\t bIsClient: %s, bIsServer: %s"), bIsClient ? TEXT("true") : TEXT("false"), bIsServer ? TEXT("true") : TEXT("false"));
+
+	// Sanity check on input params
+	if (InRemoteName.IsEmpty())
+	{
+		DoneDelegate.ExecuteIfBound(CreateErrorResponse(LOCTEXT("Error_TypeError_EmptyRemoteName", "Compare with remote called with an empty remote name")));
+		return;
+	}
+
+	if (InPackageNames.IsEmpty())
+	{
+		DoneDelegate.ExecuteIfBound(CreateErrorResponse(LOCTEXT("Error_TypeError_EmptyPackageNames", "Compare with remote called with an empty list of package names")));
+		return;
+	}
+
+	FText ErrorText;
+	FMessageAddress MessageAddress;
+	// Get address from either client or server user data, depending on local client / server state
+	if (!GetAddressFromUserData(InRemoteName, UE::StormSync::AvaBridgeCommon::StormSyncClientAddressKey, MessageAddress, &ErrorText))
+	{
+		STORM_SYNC_AVA_LOG(Error, TEXT("FStormSyncAvaSyncProvider::CompareWithRemote - %s"), *ErrorText.ToString());
+		DoneDelegate.ExecuteIfBound(CreateErrorResponse(ErrorText));
+		return;
+	}
+	
+	// Issue request and build a response
+	CompareWith(MessageAddress, InPackageNames, DoneDelegate);
+}
+
+FOnAvaSyncPackageModified& FStormSyncAvaSyncProvider::GetOnAvaSyncPackageModified()
+{
+	return OnAvaSyncPackageModified;
+}
+
+void FStormSyncAvaSyncProvider::CompareWith(const FMessageAddress& InRemoteAddress, const TArray<FName>& InPackageNames, const FOnAvaSyncCompareResponse& DoneDelegate)
+{
+	STORM_SYNC_AVA_LOG(Display, TEXT("FStormSyncAvaSyncProvider::CompareWith InPackageNames: %d (Remote Address: %s)"), InPackageNames.Num(), *InRemoteAddress.ToString());
+
+	for (const FName& PackageName : InPackageNames)
+	{
+		STORM_SYNC_AVA_LOG(Display, TEXT("\t PackageNames: %s"), *PackageName.ToString());
+	}
+
+	IStormSyncTransportClientModule::Get().RequestPackagesStatus(InRemoteAddress, InPackageNames, FOnStormSyncRequestStatusComplete::CreateLambda([DoneDelegate](const TSharedPtr<FStormSyncTransportStatusResponse>& Response)
+	{
+		if (!Response.IsValid())
+		{
+			DoneDelegate.ExecuteIfBound(CreateErrorResponse(LOCTEXT("Error_CompareInvalidResponse", "Received invalid shared ptr response")));
+			return;
+		}
+
+		STORM_SYNC_AVA_LOG(Display, TEXT("FStormSyncAvaSyncProvider::CompareWith - Response: %s"), *Response->ToString());
+
+		if (DoneDelegate.IsBound())
+		{
+			FAvaSyncCompareResponse CompareResponse;
+			CompareResponse.Status = EAvaSyncResponseResult::Success;
+			CompareResponse.ConnectionInfo = ConvertConnectionInfo(Response->ConnectionInfo);
+			CompareResponse.bNeedsSynchronization = Response->bNeedsSynchronization;
+			DoneDelegate.Execute(MakeShared<FAvaSyncCompareResponse>(CompareResponse));
+		}
+	}));
+}
+
+void FStormSyncAvaSyncProvider::OnPakAssetExtracted(const FName& InPackageName, const FString& InDestFilepath)
+{
+	OnAvaSyncPackageModified.Broadcast(InPackageName);
+}
+
+bool FStormSyncAvaSyncProvider::GetAddressFromUserData(const FString& InRemoteName, const FString& InUserDataKey, FMessageAddress& OutAddress, FText* OutErrorMessage)
+{
+	const bool bIsClient = IsPlaybackClient();
+	const bool bIsServer = IsPlaybackServer();
+
+	STORM_SYNC_AVA_LOG(
+		Display,
+		TEXT("FStormSyncAvaSyncProvider::GetAddressFromUserData InRemoteName: %s, bIsClient: %s, bIsServer: %s"),
+		*InRemoteName,
+		bIsClient ? TEXT("true") : TEXT("false"),
+		bIsServer ? TEXT("true") : TEXT("false")
+	);
+
+	if (bIsClient)
+	{
+		// Also handle the case of both a client & server, in which case, the client takes priority and server side is ignored
+		return GetAddressFromServerUserData(InRemoteName, InUserDataKey, OutAddress, OutErrorMessage);
+	}
+
+	if (bIsServer)
+	{
+		return GetAddressFromClientUserData(InRemoteName, InUserDataKey, OutAddress, OutErrorMessage);
+	}
+
+	// Case of neither a client or server
+	if (OutErrorMessage)
+	{
+		*OutErrorMessage = LOCTEXT("Error_NoClientOrServer", "Local editor instance is not a valid client or server (not started)");
+	}
+
+	FMessageAddress Address;
+	Address.Invalidate();
+	OutAddress = Address;
+	return false;
+}
+
+bool FStormSyncAvaSyncProvider::GetAddressFromServerUserData(const FString& InServerName, const FString& InUserDataKey, FMessageAddress& OutAddress, FText* OutErrorMessage)
+{
+	if (!IAvaMediaModule::Get().IsMediaPlaybackClientStarted())
+	{
+		if (OutErrorMessage)
+		{
+			*OutErrorMessage = LOCTEXT("Error_ReadingServerData_ClientNotStarted", "Trying to read server data from client but client is not started");
+		}
+		return false;
+	}
+
+	const IAvaMediaPlaybackClient& PlaybackClient = IAvaMediaModule::Get().GetMediaPlaybackClient();
+	const FString AddressId = PlaybackClient.GetServerUserData(InServerName, InUserDataKey);
+	if (AddressId.IsEmpty())
+	{
+		if (OutErrorMessage)
+		{
+			*OutErrorMessage = FText::Format(
+				LOCTEXT("Error_ReadingServerData_Empty", "Server user data Address Id is empty (ServerName: {0}, Key: {1})"),
+				FText::FromString(InServerName),
+				FText::FromString(InUserDataKey)
+			);
+		}
+		return false;
+	}
+
+	return FMessageAddress::Parse(AddressId, OutAddress);
+}
+
+bool FStormSyncAvaSyncProvider::GetAddressFromClientUserData(const FString& InClientName, const FString& InUserDataKey, FMessageAddress& OutAddress, FText* OutErrorMessage)
+{
+	if (!IAvaMediaModule::Get().IsMediaPlaybackServerStarted())
+	{
+		if (OutErrorMessage)
+		{
+			*OutErrorMessage = LOCTEXT("Error_ReadingClientData_ServerNotStarted", "Trying to read client data from server but server is not started");
+		}
+		return false;
+	}
+
+	const TSharedPtr<FAvaMediaPlaybackServer> PlaybackServer = IAvaMediaModule::Get().GetMediaPlaybackServer();
+	if (!PlaybackServer.IsValid())
+	{
+		if (OutErrorMessage)
+		{
+			*OutErrorMessage = LOCTEXT("Error_ReadingClientData_InvalidPlaybackServer", "Trying to read client data from server but server is not valid");
+		}
+		return false;
+	}
+
+	const FString AddressId = PlaybackServer->GetClientUserData(InClientName, InUserDataKey);
+	if (AddressId.IsEmpty())
+	{
+		if (OutErrorMessage)
+		{
+			*OutErrorMessage = FText::Format(
+				LOCTEXT("Error_ReadingClientData_Empty", "Client user data Address Id is empty (ClientName: {0}, Key: {1})"),
+				FText::FromString(InClientName),
+				FText::FromString(InUserDataKey)
+			);
+		}
+		return false;
+	}
+
+	return FMessageAddress::Parse(AddressId, OutAddress);
+}
+
+bool FStormSyncAvaSyncProvider::IsPlaybackServer()
+{
+	return IAvaMediaModule::Get().IsMediaPlaybackServerStarted();
+}
+
+bool FStormSyncAvaSyncProvider::IsPlaybackClient()
+{
+	return IAvaMediaModule::Get().IsMediaPlaybackClientStarted();
+}
+
+TSharedPtr<FAvaSyncCompareResponse> FStormSyncAvaSyncProvider::CreateErrorResponse(const FText& InText)
+{
+	FAvaSyncCompareResponse Response;
+	Response.Status = EAvaSyncResponseResult::Error;
+	Response.ErrorText = InText;
+	return MakeShared<FAvaSyncCompareResponse>(Response);
+}
+
+FAvaSyncResponse FStormSyncAvaSyncProvider::ConvertSyncResponse(const TSharedPtr<FStormSyncTransportSyncResponse>& InResponse)
+{
+	check(InResponse.IsValid());
+	
+	FAvaSyncResponse AvaSyncResponse;
+	if (InResponse->Status == EStormSyncResponseResult::Success)
+	{
+		AvaSyncResponse.Status = EAvaSyncResponseResult::Success;
+		AvaSyncResponse.StatusText = InResponse->StatusText;
+	}
+	else if (InResponse->Status == EStormSyncResponseResult::Error)
+	{
+		AvaSyncResponse.Status = EAvaSyncResponseResult::Error;
+		AvaSyncResponse.ErrorText = InResponse->StatusText;
+	}
+
+	AvaSyncResponse.ConnectionInfo = ConvertConnectionInfo(InResponse->ConnectionInfo);
+	return AvaSyncResponse;
+}
+
+FAvaSyncConnectionInfo FStormSyncAvaSyncProvider::ConvertConnectionInfo(const FStormSyncConnectionInfo& InConnectionInfo)
+{
+	FAvaSyncConnectionInfo ConnectionInfo;
+	ConnectionInfo.EngineVersion = InConnectionInfo.EngineVersion;
+	ConnectionInfo.InstanceId = InConnectionInfo.InstanceId;
+	ConnectionInfo.SessionId = InConnectionInfo.SessionId;
+	ConnectionInfo.HostName = InConnectionInfo.HostName;
+	ConnectionInfo.ProjectName = InConnectionInfo.ProjectName;
+	ConnectionInfo.ProjectDir = InConnectionInfo.ProjectDir;
+	
+	// Deal with instance type different enum types
+	ConnectionInfo.InstanceType = ConvertInstanceType(InConnectionInfo.InstanceType);
+	return ConnectionInfo;
+}
+
+EAvaSyncEngineType FStormSyncAvaSyncProvider::ConvertInstanceType(const EStormSyncEngineType InInstanceType)
+{
+	EAvaSyncEngineType Result;
+	switch (InInstanceType)
+	{
+	case EStormSyncEngineType::Server:
+		Result = EAvaSyncEngineType::Server;
+		break;
+	case EStormSyncEngineType::Commandlet:
+		Result = EAvaSyncEngineType::Commandlet;
+		break;
+	case EStormSyncEngineType::Editor:
+		Result = EAvaSyncEngineType::Editor;
+		break;
+	case EStormSyncEngineType::Game:
+		Result = EAvaSyncEngineType::Game;
+		break;
+	case EStormSyncEngineType::Other:
+		Result = EAvaSyncEngineType::Other;
+		break;
+	case EStormSyncEngineType::Unknown:
+		Result = EAvaSyncEngineType::Unknown;
+		break;
+	default:
+		Result = EAvaSyncEngineType::Unknown;
+		break;
+	}
+
+	return Result;
+}
+
+#undef LOCTEXT_NAMESPACE
diff --git a/Engine/Plugins/Experimental/StormSyncAvalancheBridge/Source/StormSyncAvaBridge/Private/ModularFeature/StormSyncAvaSyncProvider.h b/Engine/Plugins/Experimental/StormSyncAvalancheBridge/Source/StormSyncAvaBridge/Private/ModularFeature/StormSyncAvaSyncProvider.h
new file mode 100644
--- /dev/null
+++ b/Engine/Plugins/Experimental/StormSyncAvalancheBridge/Source/StormSyncAvaBridge/Private/ModularFeature/StormSyncAvaSyncProvider.h
@@ -0,0 +1,81 @@
+// Copyright Epic Games, Inc. All Rights Reserved.
+
+#pragma once
+
+#include "CoreMinimal.h"
+#include "ModularFeature/IAvaMediaSyncProvider.h"
+
+enum class EStormSyncEngineType : uint8;
+struct FMessageAddress;
+struct FStormSyncConnectionInfo;
+struct FStormSyncFileDependency;
+struct FStormSyncTransportStatusResponse;
+struct FStormSyncTransportSyncResponse;
+
+/**
+ * Avalanche synchronization provider feature implementation
+ */
+class FStormSyncAvaSyncProvider : public IAvaMediaSyncProvider
+{
+public:
+	FStormSyncAvaSyncProvider();
+	virtual ~FStormSyncAvaSyncProvider() override;
+
+	//~ Begin IAvaMediaSyncProvider interface
+	virtual FName GetName() const override;
+	virtual void SyncToAll(const TArray<FName>& InPackageNames) override;
+	virtual void PushToRemote(const FString& InRemoteName, const TArray<FName>& InPackageNames, const FOnAvaSyncResponse& DoneDelegate) override;
+	virtual void PullFromRemote(const FString& InRemoteName, const TArray<FName>& InPackageNames, const FOnAvaSyncResponse& DoneDelegate) override;
+	virtual void CompareWithRemote(const FString& InRemoteName, const TArray<FName>& InPackageNames, const FOnAvaSyncCompareResponse& DoneDelegate) override;
+	virtual FOnAvaSyncPackageModified& GetOnAvaSyncPackageModified() override;
+	//~ End IAvaMediaSyncProvider interface
+
+protected:
+	/** Low lvl API using direct message address destination */
+	virtual void CompareWith(const FMessageAddress& InRemoteAddress, const TArray<FName>& InPackageNames, const FOnAvaSyncCompareResponse& DoneDelegate);
+
+	/**
+	 * Handler called when a file is extracted from an incoming storm sync pak
+	 * and relay to the OnAvaSyncPackageModified event
+	 */
+	void OnPakAssetExtracted(const FName& InPackageName, const FString& InDestFilepath);
+	
+private:
+	FOnAvaSyncPackageModified OnAvaSyncPackageModified;
+
+	/**
+	 * Helper to get either client or server user data depending on local playback client / server state
+	 *
+	 * eg.
+	 *
+	 * 1. If client, we get user data from server
+	 * 2. If server, we get user data from client
+	 * 3. In case we are neither a client of server, returns invalid address and false
+	 * 4. If an instance is both a client and server, the client takes priority and server side is ignored
+	 */
+	static bool GetAddressFromUserData(const FString& InRemoteName, const FString& InUserDataKey, FMessageAddress& OutAddress, FText* OutErrorMessage = nullptr);
+	
+	/** Helper to get server user data from playback client */
+	static bool GetAddressFromServerUserData(const FString& InServerName, const FString& InUserDataKey, FMessageAddress& OutAddress, FText* OutErrorMessage = nullptr);
+	
+	/** Helper to get client user data from playback server */
+	static bool GetAddressFromClientUserData(const FString& InClientName, const FString& InUserDataKey, FMessageAddress& OutAddress, FText* OutErrorMessage = nullptr);
+
+	/** Returns whether local editor instance has a playback server started */
+	static bool IsPlaybackServer();
+	
+	/** Returns whether local editor instance has a playback client started */
+	static bool IsPlaybackClient();
+
+	/** Helper to create a new error response (used when calling back DoneDelegate to indicate an error) */
+	static TSharedPtr<FAvaSyncCompareResponse> CreateErrorResponse(const FText& InText);
+
+	/** Helper to convert a FStormSyncTransportSyncResponse struct to the modular feature equivalent */
+	static FAvaSyncResponse ConvertSyncResponse(const TSharedPtr<FStormSyncTransportSyncResponse>& InResponse);
+	
+	/** Helper to convert a FStormSyncConnectionInfo struct to the modular feature equivalent */
+	static FAvaSyncConnectionInfo ConvertConnectionInfo(const FStormSyncConnectionInfo& InConnectionInfo);
+
+	/** Helper to convert instance type enum specific to this implementation to the interface one (for connection info) */
+	static EAvaSyncEngineType ConvertInstanceType(const EStormSyncEngineType InInstanceType);
+};
diff --git a/Engine/Plugins/Experimental/StormSyncAvalancheBridge/Source/StormSyncAvaBridge/Private/StormSyncAvaBridge.cpp b/Engine/Plugins/Experimental/StormSyncAvalancheBridge/Source/StormSyncAvaBridge/Private/StormSyncAvaBridge.cpp
new file mode 100644
--- /dev/null
+++ b/Engine/Plugins/Experimental/StormSyncAvalancheBridge/Source/StormSyncAvaBridge/Private/StormSyncAvaBridge.cpp
@@ -0,0 +1,231 @@
+// Copyright Epic Games, Inc. All Rights Reserved.
+
+#include "StormSyncAvaBridge.h"
+
+#include "Features/IModularFeatures.h"
+#include "IAvaMediaModule.h"
+#include "IStormSyncTransportClientModule.h"
+#include "IStormSyncTransportServerModule.h"
+#include "Misc/CoreDelegates.h"
+#include "ModularFeature/StormSyncAvaSyncProvider.h"
+#include "Playback/AvaMediaPlaybackServer.h"
+#include "Playback/IAvaMediaPlaybackClient.h"
+#include "StormSyncAvaBridgeCommon.h"
+#include "StormSyncAvaBridgeLog.h"
+#include "StormSyncAvaBridgeUtils.h"
+#include "StormSyncCoreDelegates.h"
+
+#define LOCTEXT_NAMESPACE "FStormSyncAvaBridgeModule"
+
+void FStormSyncAvaBridgeModule::StartupModule()
+{
+	FCoreDelegates::OnPostEngineInit.AddRaw(this, &FStormSyncAvaBridgeModule::OnPostEngineInit);
+	FStormSyncCoreDelegates::OnStormSyncServerStarted.AddRaw(this, &FStormSyncAvaBridgeModule::OnStormSyncServerStarted);
+	FStormSyncCoreDelegates::OnStormSyncServerStopped.AddRaw(this, &FStormSyncAvaBridgeModule::OnStormSyncServerStopped);
+
+	// Instantiate the sync provider feature on module start
+	SyncProvider = MakeUnique<FStormSyncAvaSyncProvider>();
+
+	RegisterConsoleCommands();
+}
+
+void FStormSyncAvaBridgeModule::ShutdownModule()
+{
+	FCoreDelegates::OnPostEngineInit.RemoveAll(this);
+	FStormSyncCoreDelegates::OnStormSyncServerStarted.RemoveAll(this);
+	FStormSyncCoreDelegates::OnStormSyncServerStopped.RemoveAll(this);
+
+	UnregisterConsoleCommands();
+
+	// Unregister the interceptor feature on module shutdown
+	if (SyncProvider.IsValid())
+	{
+		IModularFeatures::Get().UnregisterModularFeature(IAvaMediaSyncProvider::GetModularFeatureName(), SyncProvider.Get());
+	}
+}
+
+void FStormSyncAvaBridgeModule::OnPostEngineInit()
+{
+	if (IAvaMediaModule::IsModuleLoaded())
+	{
+		if (IAvaMediaModule::Get().IsMediaPlaybackServerStarted())
+		{
+			RegisterUserDataForPlaybackServer();
+		}
+
+		if (IAvaMediaModule::Get().IsMediaPlaybackClientStarted())
+		{
+			RegisterUserDataForPlaybackClient();
+		}
+
+		IAvaMediaModule::Get().GetOnAvaMediaPlaybackServerStarted().AddStatic(&FStormSyncAvaBridgeModule::RegisterUserDataForPlaybackServer);
+		IAvaMediaModule::Get().GetOnAvaMediaPlaybackClientStarted().AddStatic(&FStormSyncAvaBridgeModule::RegisterUserDataForPlaybackClient);
+	}
+
+	if (IStormSyncTransportServerModule::IsAvailable())
+	{
+		RegisterUserDataForPlaybackServer();
+	}
+
+	// Register the sync provider feature
+	if (SyncProvider.IsValid())
+	{
+		IModularFeatures::Get().RegisterModularFeature(IAvaMediaSyncProvider::GetModularFeatureName(), SyncProvider.Get());
+	}
+}
+
+void FStormSyncAvaBridgeModule::RegisterConsoleCommands()
+{
+	ConsoleCommands.Add(IConsoleManager::Get().RegisterConsoleCommand(
+		TEXT("StormSyncAvaBridge.Debug.GetUserData"),
+		TEXT("Returns the associated value in Playback Server for ChannelName and Data Key"),
+		FConsoleCommandWithArgsDelegate::CreateStatic(&FStormSyncAvaBridgeModule::ExecuteGetUserData),
+		ECVF_Default
+	));
+}
+
+void FStormSyncAvaBridgeModule::UnregisterConsoleCommands()
+{
+	for (IConsoleObject* Cmd : ConsoleCommands)
+	{
+		IConsoleManager::Get().UnregisterConsoleObject(Cmd);
+	}
+
+	ConsoleCommands.Empty();
+}
+
+void FStormSyncAvaBridgeModule::ExecuteGetUserData(const TArray<FString>& Args)
+{
+	const FString Arguments = FString::Join(Args, TEXT(" "));
+	STORM_SYNC_AVA_LOG(Display, TEXT("FStormSyncAvaBridgeModule::ExecuteGetUserData - Args: %s"), *Arguments)
+
+	if (!Args.IsValidIndex(0) || Args[0].IsEmpty())
+	{
+		STORM_SYNC_AVA_LOG(Error, TEXT("Missing first parameter \"ChannelName\""))
+		return;
+	}
+
+	if (!Args.IsValidIndex(1) || Args[1].IsEmpty())
+	{
+		STORM_SYNC_AVA_LOG(Error, TEXT("Missing second parameter \"Key\""))
+		return;
+	}
+
+	if (!IAvaMediaModule::IsModuleLoaded())
+	{
+		STORM_SYNC_AVA_LOG(Error, TEXT("Ava Media Module is not available"))
+		return;
+	}
+
+	if (!IAvaMediaModule::Get().IsMediaPlaybackClientStarted())
+	{
+		STORM_SYNC_AVA_LOG(Error, TEXT("Ava Playback Client stopped"))
+		return;
+	}
+
+	const FString ChannelName = Args[0];
+	const FString Key = Args[1];
+
+	TArray<FString> ServerNames = FStormSyncAvaBridgeUtils::GetServerNamesForChannel(ChannelName);
+	if (ServerNames.IsEmpty())
+	{
+		STORM_SYNC_AVA_LOG(Display, TEXT("Ava Broadcast profile for channel %s has no remotes"), *ChannelName)
+		return;
+	}
+
+	const IAvaMediaPlaybackClient& PlaybackClient = IAvaMediaModule::Get().GetMediaPlaybackClient();
+	for (const FString& ServerName : ServerNames)
+	{
+		FString ServerValue = PlaybackClient.GetServerUserData(ServerName, Key);
+		STORM_SYNC_AVA_LOG(Display, TEXT("\t Media Playback Client User Data for server %s - %s:%s"), *ServerName, *Key, *ServerValue);
+	}
+}
+
+void FStormSyncAvaBridgeModule::OnStormSyncServerStarted()
+{
+	STORM_SYNC_AVA_LOG(Verbose, TEXT("FStormSyncAvaBridgeModule::OnStormSyncServerStarted"))
+	RegisterUserDataForPlaybackServer();
+}
+
+void FStormSyncAvaBridgeModule::OnStormSyncServerStopped()
+{
+	STORM_SYNC_AVA_LOG(Verbose, TEXT("FStormSyncAvaBridgeModule::OnStormSyncServerStopped"))
+	RegisterUserDataForPlaybackServer();
+}
+
+void FStormSyncAvaBridgeModule::RegisterUserDataForPlaybackServer()
+{
+	// Fail safe checks for both Avalanche and Storm Sync modules, since this method can be executed from delegates in either of these modules
+	if (!ValidateModulesAreAvailable())
+	{
+		return;
+	}
+
+	if (const TSharedPtr<FAvaMediaPlaybackServer> PlaybackServer = IAvaMediaModule::Get().GetMediaPlaybackServer())
+	{
+		const FString StormSyncServerAddress = IStormSyncTransportServerModule::Get().GetServerEndpointMessageAddressId();
+		const FString DiscoveryManagerAddress = IStormSyncTransportServerModule::Get().GetDiscoveryManagerMessageAddressId();
+		const FString StormSyncClientAddress = IStormSyncTransportClientModule::Get().GetClientEndpointMessageAddressId();
+
+		PlaybackServer->SetUserData(UE::StormSync::AvaBridgeCommon::StormSyncServerAddressKey, StormSyncServerAddress);
+		PlaybackServer->SetUserData(UE::StormSync::AvaBridgeCommon::StormSyncClientAddressKey, StormSyncClientAddress);
+		PlaybackServer->SetUserData(UE::StormSync::AvaBridgeCommon::StormSyncDiscoveryAddressKey, DiscoveryManagerAddress);
+
+		STORM_SYNC_AVA_LOG(Display, TEXT("FStormSyncAvaBridgeModule::RegisterUserDataForPlaybackServer StormSyncServerAddress: %s"), *StormSyncServerAddress)
+		STORM_SYNC_AVA_LOG(Display, TEXT("FStormSyncAvaBridgeModule::RegisterUserDataForPlaybackServer StormSyncClientAddress: %s"), *StormSyncClientAddress)
+		STORM_SYNC_AVA_LOG(Display, TEXT("FStormSyncAvaBridgeModule::RegisterUserDataForPlaybackServer DiscoveryManagerAddress: %s"), *DiscoveryManagerAddress)
+	}
+}
+
+void FStormSyncAvaBridgeModule::RegisterUserDataForPlaybackClient()
+{
+	// Fail safe checks for both Avalanche and Storm Sync modules, since this method can be executed from delegates in either of these modules
+	if (!ValidateModulesAreAvailable())
+	{
+		return;
+	}
+
+	if (IAvaMediaModule::Get().IsMediaPlaybackClientStarted())
+	{
+		IAvaMediaPlaybackClient& PlaybackClient = IAvaMediaModule::Get().GetMediaPlaybackClient();
+
+		const FString StormSyncServerAddress = IStormSyncTransportServerModule::Get().GetServerEndpointMessageAddressId();
+		const FString DiscoveryManagerAddress = IStormSyncTransportServerModule::Get().GetDiscoveryManagerMessageAddressId();
+		const FString StormSyncClientAddress = IStormSyncTransportClientModule::Get().GetClientEndpointMessageAddressId();
+
+		PlaybackClient.SetUserData(UE::StormSync::AvaBridgeCommon::StormSyncServerAddressKey, StormSyncServerAddress);
+		PlaybackClient.SetUserData(UE::StormSync::AvaBridgeCommon::StormSyncClientAddressKey, StormSyncClientAddress);
+		PlaybackClient.SetUserData(UE::StormSync::AvaBridgeCommon::StormSyncDiscoveryAddressKey, DiscoveryManagerAddress);
+
+		STORM_SYNC_AVA_LOG(Display, TEXT("FStormSyncAvaBridgeModule::RegisterUserDataForPlaybackClient StormSyncServerAddress: %s"), *StormSyncServerAddress)
+		STORM_SYNC_AVA_LOG(Display, TEXT("FStormSyncAvaBridgeModule::RegisterUserDataForPlaybackClient StormSyncClientAddress: %s"), *StormSyncClientAddress)
+		STORM_SYNC_AVA_LOG(Display, TEXT("FStormSyncAvaBridgeModule::RegisterUserDataForPlaybackClient DiscoveryManagerAddress: %s"), *DiscoveryManagerAddress)
+	}
+}
+
+bool FStormSyncAvaBridgeModule::ValidateModulesAreAvailable()
+{
+	if (!IAvaMediaModule::IsModuleLoaded())
+	{
+		STORM_SYNC_AVA_LOG(Warning, TEXT("FStormSyncAvaBridgeModule::ValidateModulesAreAvailable - Failed to set user data cause Ava Media Module is not loaded"))
+		return false;
+	}
+
+	if (!IStormSyncTransportServerModule::IsAvailable())
+	{
+		STORM_SYNC_AVA_LOG(Warning, TEXT("FStormSyncAvaBridgeModule::ValidateModulesAreAvailable - Failed to set user data cause Storm Sync Server Module is not loaded"))
+		return false;
+	}
+
+	if (!IStormSyncTransportClientModule::IsAvailable())
+	{
+		STORM_SYNC_AVA_LOG(Warning, TEXT("FStormSyncAvaBridgeModule::ValidateModulesAreAvailable - Failed to set user data cause Storm Sync Server Client is not loaded"))
+		return false;
+	}
+
+	return true;
+}
+
+#undef LOCTEXT_NAMESPACE
+
+IMPLEMENT_MODULE(FStormSyncAvaBridgeModule, StormSyncAvaBridge)
diff --git a/Engine/Plugins/Experimental/StormSyncAvalancheBridge/Source/StormSyncAvaBridge/Private/StormSyncAvaBridgeLog.cpp b/Engine/Plugins/Experimental/StormSyncAvalancheBridge/Source/StormSyncAvaBridge/Private/StormSyncAvaBridgeLog.cpp
new file mode 100644
--- /dev/null
+++ b/Engine/Plugins/Experimental/StormSyncAvalancheBridge/Source/StormSyncAvaBridge/Private/StormSyncAvaBridgeLog.cpp
@@ -0,0 +1,5 @@
+// Copyright Epic Games, Inc. All Rights Reserved.
+
+#include "StormSyncAvaBridgeLog.h"
+
+DEFINE_LOG_CATEGORY(LogStormSyncAvaBridge);
diff --git a/Engine/Plugins/Experimental/StormSyncAvalancheBridge/Source/StormSyncAvaBridge/Private/StormSyncAvaBridgeUtils.cpp b/Engine/Plugins/Experimental/StormSyncAvalancheBridge/Source/StormSyncAvaBridge/Private/StormSyncAvaBridgeUtils.cpp
new file mode 100644
--- /dev/null
+++ b/Engine/Plugins/Experimental/StormSyncAvalancheBridge/Source/StormSyncAvaBridge/Private/StormSyncAvaBridgeUtils.cpp
@@ -0,0 +1,30 @@
+// Copyright Epic Games, Inc. All Rights Reserved.
+
+
+#include "StormSyncAvaBridgeUtils.h"
+
+#include "AvalancheBroadcast.h"
+
+TArray<FString> FStormSyncAvaBridgeUtils::GetServerNamesForChannel(const FString& InChannelName)
+{
+	TArray<FString> ServerNames;
+	
+	if (UAvalancheBroadcast* Broadcast = UAvalancheBroadcast::GetAvalancheBroadcast())
+	{
+		const FAvaOutputChannel Channel = Broadcast->GetCurrentProfile().GetChannel(FName(*InChannelName));
+		if (Channel.IsValidChannel())
+		{
+			TArray<UMediaOutput*> Outputs = Channel.GetMediaOutputs();
+			for (const UMediaOutput* Output : Outputs)
+			{
+				FAvaMediaOutputInfo OutputInfo = Channel.GetMediaOutputInfo(Output);
+				if (OutputInfo.IsValid() && OutputInfo.IsRemote())
+				{
+					ServerNames.Add(OutputInfo.ServerName);
+				}
+			}
+		}
+	}
+
+	return ServerNames;
+}
diff --git a/Engine/Plugins/Experimental/StormSyncAvalancheBridge/Source/StormSyncAvaBridge/Public/StormSyncAvaBridge.h b/Engine/Plugins/Experimental/StormSyncAvalancheBridge/Source/StormSyncAvaBridge/Public/StormSyncAvaBridge.h
new file mode 100644
--- /dev/null
+++ b/Engine/Plugins/Experimental/StormSyncAvalancheBridge/Source/StormSyncAvaBridge/Public/StormSyncAvaBridge.h
@@ -0,0 +1,68 @@
+// Copyright Epic Games, Inc. All Rights Reserved.
+
+#pragma once
+
+#include "CoreMinimal.h"
+#include "Modules/ModuleInterface.h"
+
+class FStormSyncAvaSyncProvider;
+
+/**
+ * Module for Storm Sync Ava Bridge runtime module
+ *
+ * This class handles the update of User Data in Ava Media Playback Server when module is loaded,
+ * and Storm Sync server is started or stopped.
+ *
+ * This is then replicated and made available on remote clients so that a Playlist asset can associate
+ * a message bus address id (for storm sync push / sync) to a remote on a channel.
+ */
+class FStormSyncAvaBridgeModule : public IModuleInterface
+{
+public:
+	//~ Begin IModuleInterface interface
+	virtual void StartupModule() override;
+	virtual void ShutdownModule() override;
+	//~ End IModuleInterface interface
+
+private:
+	/** References of registered console commands via IConsoleManager */
+	TArray<IConsoleObject*> ConsoleCommands;
+
+	/** Ava sync provider feature instance */
+	TUniquePtr<FStormSyncAvaSyncProvider> SyncProvider;
+
+	/** Updates playback user data when engine is fully loaded to setup storm sync message address ids */
+	void OnPostEngineInit();
+
+	/** Called from StartupModule and sets up console commands for the plugin via IConsoleManager */
+	void RegisterConsoleCommands();
+
+	/** Called from ShutdownModule and clears out previously registered console commands */
+	void UnregisterConsoleCommands();
+
+	/** Console command handler to return associated user data in ava playback server */
+	static void ExecuteGetUserData(const TArray<FString>& Args);
+
+	/** Updates user data when storm sync server is started */
+	void OnStormSyncServerStarted();
+	
+	/** Updates user data when storm sync server is stopped */
+	void OnStormSyncServerStopped();
+
+	/**
+	 * Actual helper to set the user data for a running playback server.
+	 * 
+	 * Sets up Storm Sync Server / Client and Discovery Manager address ids
+	 */
+	static void RegisterUserDataForPlaybackServer();
+	
+	/**
+	 * Actual helper to set the user data for a running playback client.
+	 * 
+	 * Sets up Storm Sync Server / Client and Discovery Manager address ids
+	 */
+	static void RegisterUserDataForPlaybackClient();
+
+	/** Helper ensuring all required modules are loaded available */
+	static bool ValidateModulesAreAvailable();
+};
diff --git a/Engine/Plugins/Experimental/StormSyncAvalancheBridge/Source/StormSyncAvaBridge/Public/StormSyncAvaBridgeCommon.h b/Engine/Plugins/Experimental/StormSyncAvalancheBridge/Source/StormSyncAvaBridge/Public/StormSyncAvaBridgeCommon.h
new file mode 100644
--- /dev/null
+++ b/Engine/Plugins/Experimental/StormSyncAvalancheBridge/Source/StormSyncAvaBridge/Public/StormSyncAvaBridgeCommon.h
@@ -0,0 +1,17 @@
+// Copyright Epic Games, Inc. All Rights Reserved.
+
+#pragma once
+
+#include "CoreMinimal.h"
+
+namespace UE::StormSync::AvaBridgeCommon
+{
+	/** The user data key to store storm sync server address id */
+	static constexpr const TCHAR* StormSyncServerAddressKey = TEXT("StormSyncServerAddress");
+	
+	/** The user data key to store storm sync discovery manager address id */
+	static constexpr const TCHAR* StormSyncClientAddressKey = TEXT("StormSyncClientAddress");
+
+	/** The user data key to store storm sync discovery manager address id */
+	static constexpr const TCHAR* StormSyncDiscoveryAddressKey = TEXT("StormSyncDiscoveryAddress");
+}
\ No newline at end of file
diff --git a/Engine/Plugins/Experimental/StormSyncAvalancheBridge/Source/StormSyncAvaBridge/Public/StormSyncAvaBridgeLog.h b/Engine/Plugins/Experimental/StormSyncAvalancheBridge/Source/StormSyncAvaBridge/Public/StormSyncAvaBridgeLog.h
new file mode 100644
--- /dev/null
+++ b/Engine/Plugins/Experimental/StormSyncAvalancheBridge/Source/StormSyncAvaBridge/Public/StormSyncAvaBridgeLog.h
@@ -0,0 +1,12 @@
+// Copyright Epic Games, Inc. All Rights Reserved.
+
+#pragma once
+
+#include "CoreMinimal.h"
+
+STORMSYNCAVABRIDGE_API DECLARE_LOG_CATEGORY_EXTERN(LogStormSyncAvaBridge, Display, All);
+
+#define STORM_SYNC_AVA_LOG(Verbosity, Format, ...) \
+{ \
+    UE_LOG(LogStormSyncAvaBridge, Verbosity, Format, ##__VA_ARGS__); \
+}
diff --git a/Engine/Plugins/Experimental/StormSyncAvalancheBridge/Source/StormSyncAvaBridge/Public/StormSyncAvaBridgeUtils.h b/Engine/Plugins/Experimental/StormSyncAvalancheBridge/Source/StormSyncAvaBridge/Public/StormSyncAvaBridgeUtils.h
new file mode 100644
--- /dev/null
+++ b/Engine/Plugins/Experimental/StormSyncAvalancheBridge/Source/StormSyncAvaBridge/Public/StormSyncAvaBridgeUtils.h
@@ -0,0 +1,12 @@
+// Copyright Epic Games, Inc. All Rights Reserved.
+
+#pragma once
+
+#include "CoreMinimal.h"
+
+class STORMSYNCAVABRIDGE_API FStormSyncAvaBridgeUtils
+{
+public:
+	/** Returns the list of Ava Playback Server nmes for a given channel */
+	static TArray<FString> GetServerNamesForChannel(const FString& InChannelName);
+};
diff --git a/Engine/Plugins/Experimental/StormSyncAvalancheBridge/Source/StormSyncAvaBridge/StormSyncAvaBridge.Build.cs b/Engine/Plugins/Experimental/StormSyncAvalancheBridge/Source/StormSyncAvaBridge/StormSyncAvaBridge.Build.cs
new file mode 100644
--- /dev/null
+++ b/Engine/Plugins/Experimental/StormSyncAvalancheBridge/Source/StormSyncAvaBridge/StormSyncAvaBridge.Build.cs
@@ -0,0 +1,33 @@
+// Copyright Epic Games, Inc. All Rights Reserved.
+
+using UnrealBuildTool;
+
+public class StormSyncAvaBridge : ModuleRules
+{
+	public StormSyncAvaBridge(ReadOnlyTargetRules Target) : base(Target)
+	{
+		PCHUsage = ModuleRules.PCHUsageMode.UseExplicitOrSharedPCHs;
+
+		PublicDependencyModuleNames.AddRange(
+			new string[]
+			{
+				"Core",
+			}
+		);
+
+		PrivateDependencyModuleNames.AddRange(
+			new string[]
+			{
+				"AvalancheMedia",
+				"CoreUObject",
+				"Engine",
+				"Slate",
+				"SlateCore",
+				"StormSyncCore",
+				"StormSyncTransportCore",
+				"StormSyncTransportServer",
+				"StormSyncTransportClient",
+			}
+		);
+	}
+}
diff --git a/Engine/Plugins/Experimental/StormSyncAvalancheBridge/Source/StormSyncAvaBridgeEditor/Private/StormSyncAvaBridgeEditor.cpp b/Engine/Plugins/Experimental/StormSyncAvalancheBridge/Source/StormSyncAvaBridgeEditor/Private/StormSyncAvaBridgeEditor.cpp
new file mode 100644
--- /dev/null
+++ b/Engine/Plugins/Experimental/StormSyncAvalancheBridge/Source/StormSyncAvaBridgeEditor/Private/StormSyncAvaBridgeEditor.cpp
@@ -0,0 +1,20 @@
+// Copyright Epic Games, Inc. All Rights Reserved.
+
+#include "StormSyncAvaBridgeEditor.h"
+
+#include "StormSyncAvaPlaylistExtender.h"
+
+void FStormSyncAvaBridgeEditorModule::StartupModule()
+{
+	PlaylistExtender = MakeShared<FStormSyncAvaPlaylistExtender>();
+}
+
+void FStormSyncAvaBridgeEditorModule::ShutdownModule()
+{
+	if (PlaylistExtender.IsValid())
+	{
+		PlaylistExtender.Reset();
+	}
+}
+
+IMPLEMENT_MODULE(FStormSyncAvaBridgeEditorModule, StormSyncAvaBridgeEditor)
diff --git a/Engine/Plugins/Experimental/StormSyncAvalancheBridge/Source/StormSyncAvaBridgeEditor/Private/StormSyncAvaBridgeEditorLog.cpp b/Engine/Plugins/Experimental/StormSyncAvalancheBridge/Source/StormSyncAvaBridgeEditor/Private/StormSyncAvaBridgeEditorLog.cpp
new file mode 100644
--- /dev/null
+++ b/Engine/Plugins/Experimental/StormSyncAvalancheBridge/Source/StormSyncAvaBridgeEditor/Private/StormSyncAvaBridgeEditorLog.cpp
@@ -0,0 +1,5 @@
+// Copyright Epic Games, Inc. All Rights Reserved.
+
+#include "StormSyncAvaBridgeEditorLog.h"
+
+DEFINE_LOG_CATEGORY(LogStormSyncAvaBridgeEditor);
diff --git a/Engine/Plugins/Experimental/StormSyncAvalancheBridge/Source/StormSyncAvaBridgeEditor/Private/StormSyncAvaPlaylistExtender.cpp b/Engine/Plugins/Experimental/StormSyncAvalancheBridge/Source/StormSyncAvaBridgeEditor/Private/StormSyncAvaPlaylistExtender.cpp
new file mode 100644
--- /dev/null
+++ b/Engine/Plugins/Experimental/StormSyncAvalancheBridge/Source/StormSyncAvaBridgeEditor/Private/StormSyncAvaPlaylistExtender.cpp
@@ -0,0 +1,588 @@
+// Copyright Epic Games, Inc. All Rights Reserved.
+
+#include "StormSyncAvaPlaylistExtender.h"
+
+#include "Framework/MultiBox/MultiBoxBuilder.h"
+#include "IAvaMediaEditorModule.h"
+#include "IAvaMediaModule.h"
+#include "IStormSyncTransportClientModule.h"
+#include "Playback/AvaMediaPlaybackServer.h"
+#include "Playback/IAvaMediaPlaybackClient.h"
+#include "Playlist/AvaPlaylistEditor.h"
+#include "Playlist/AvalanchePlaylist.h"
+#include "StormSyncAvaBridgeCommon.h"
+#include "StormSyncAvaBridgeEditorLog.h"
+#include "StormSyncAvaBridgeUtils.h"
+#include "StormSyncEditor.h"
+#include "Subsystems/StormSyncNotificationSubsystem.h"
+#include "Toolkits/ToolkitManager.h"
+
+#define LOCTEXT_NAMESPACE "FStormSyncAvaPlaylistExtender"
+
+FStormSyncAvaPlaylistExtender::FStormSyncAvaPlaylistExtender()
+{
+	FCoreDelegates::OnPostEngineInit.AddRaw(this, &FStormSyncAvaPlaylistExtender::OnPostEngineInit);
+}
+
+FStormSyncAvaPlaylistExtender::~FStormSyncAvaPlaylistExtender()
+{
+	FCoreDelegates::OnPostEngineInit.RemoveAll(this);
+
+	if (IAvaMediaEditorModule::IsLoaded())
+	{
+		UnregisterExtension(IAvaMediaEditorModule::Get().GetPlaylistMenuExtensibilityManager(), MenuExtenderHandle);
+		UnregisterExtension(IAvaMediaEditorModule::Get().GetPlaylistToolBarExtensibilityManager(), ToolbarExtenderHandle);
+	}
+}
+
+void FStormSyncAvaPlaylistExtender::OnPostEngineInit()
+{
+	if (IAvaMediaEditorModule::IsLoaded())
+	{
+		RegisterMenuExtensions();
+	}
+}
+
+void FStormSyncAvaPlaylistExtender::RegisterMenuExtensions()
+{
+	// Context menu extension
+	MenuExtenderHandle = RegisterExtension(
+		IAvaMediaEditorModule::Get().GetPlaylistMenuExtensibilityManager(),
+		FAssetEditorExtender::CreateSP(this, &FStormSyncAvaPlaylistExtender::AddMenuExtender)
+	);
+	
+	// Toolbar extension
+	ToolbarExtenderHandle = RegisterExtension(
+		IAvaMediaEditorModule::Get().GetPlaylistToolBarExtensibilityManager(),
+		FAssetEditorExtender::CreateSP(this, &FStormSyncAvaPlaylistExtender::AddToolbarExtender)
+	);
+}
+
+FDelegateHandle FStormSyncAvaPlaylistExtender::RegisterExtension(const TSharedPtr<FExtensibilityManager> InExtensibilityManager, const FAssetEditorExtender& InExtenderDelegate)
+{
+	if (!InExtensibilityManager.IsValid())
+	{
+		const FDelegateHandle DelegateHandle;
+		return DelegateHandle;
+	}
+
+	const int32 ExtenderIndex = InExtensibilityManager->GetExtenderDelegates().Add(InExtenderDelegate);
+	return InExtensibilityManager->GetExtenderDelegates()[ExtenderIndex].GetHandle();
+}
+
+void FStormSyncAvaPlaylistExtender::UnregisterExtension(const TSharedPtr<FExtensibilityManager> InExtensibilityManager, const FDelegateHandle& InHandleToRemove)
+{
+	if (!InExtensibilityManager.IsValid())
+	{
+		return;
+	}
+
+	InExtensibilityManager->GetExtenderDelegates().RemoveAll([Handle = InHandleToRemove](const FAssetEditorExtender& Extender)
+	{
+		return Handle == Extender.GetHandle();
+	});
+}
+
+TSharedRef<FExtender> FStormSyncAvaPlaylistExtender::AddMenuExtender(const TSharedRef<FUICommandList> InCommandList, const TArray<UObject*> ContextSensitiveObjects)
+{
+	STORM_SYNC_AVA_EDITOR_LOG(Display, TEXT("FStormSyncAvaPlaylistExtender::AddMenuExtender - Adding in menu extensions"))
+	TSharedRef<FExtender> Extender(new FExtender());
+
+	const UAvalanchePlaylist* Playlist = ContextSensitiveObjects.IsValidIndex(0) ? Cast<UAvalanchePlaylist>(ContextSensitiveObjects[0]) : nullptr;
+	if (!Playlist)
+	{
+		return Extender;
+	}
+
+	const TSharedPtr<IToolkit> AssetEditor = FToolkitManager::Get().FindEditorForAsset(Playlist);
+	if (!AssetEditor.IsValid())
+	{
+		return Extender;
+	}
+
+	const TSharedPtr<FAvaPlaylistEditor> PlaylistEditor = StaticCastSharedPtr<FAvaPlaylistEditor>(AssetEditor);
+	if (!PlaylistEditor.IsValid())
+	{
+		return Extender;
+	}
+
+	// Template panel extension
+	Extender->AddMenuExtension(
+		AvalancheExtensionHook,
+		EExtensionHook::After,
+		InCommandList,
+		// Convert to weak ptr to prevent ownership to the playlist editor and potentially increasing its lifetime
+		FMenuExtensionDelegate::CreateSP(this, &FStormSyncAvaPlaylistExtender::CreateTemplateContextMenu, Playlist, TWeakPtr<FAvaPlaylistEditor>(PlaylistEditor))
+	);
+
+	return Extender;
+}
+
+void FStormSyncAvaPlaylistExtender::CreateTemplateContextMenu(FMenuBuilder& MenuBuilder, const UAvalanchePlaylist* InPlaylist, TWeakPtr<FAvaPlaylistEditor> InPlaylistEditor)
+{
+	check(InPlaylist);
+
+	const TSharedPtr<FAvaPlaylistEditor> PlaylistEditor = InPlaylistEditor.Pin();
+	if (!PlaylistEditor.IsValid())
+	{
+		STORM_SYNC_AVA_EDITOR_LOG(Error, TEXT("FStormSyncAvaPlaylistExtender::CreateTemplateContextMenu - Invalid shared ptr from weak ptr delegate param"))
+		return;
+	}
+
+	if (!FModuleManager::Get().IsModuleLoaded(TEXT("StormSyncEditor")))
+	{
+		STORM_SYNC_AVA_EDITOR_LOG(Error, TEXT("FStormSyncAvaPlaylistExtender::CreateTemplateContextMenu - StormSyncEditor module is not loaded. Is StormSync plugin enabled ?"))
+		return;
+	}
+
+	FText DisabledTooltipReason;
+	bool bIsValidSelection = false;
+	TArray<FName> SelectedPackageNames;
+	FCanExecuteAction DefaultCanExecuteAction;
+	GetContextMenuSelectionInfos(InPlaylist, InPlaylistEditor, bIsValidSelection, DisabledTooltipReason, SelectedPackageNames, DefaultCanExecuteAction);
+
+	// For later use with action handler
+	const FStormSyncEditorModule& StormSyncEditor = FStormSyncEditorModule::Get();
+
+	MenuBuilder.BeginSection("StormSyncOperations_Template", LOCTEXT("StormSyncOperationsHeader", "Synchronize Actions"));
+	
+	// Initialize action
+	{
+		const FText TooltipText = bIsValidSelection ?
+			FText::Format(LOCTEXT("Initialize_Tooltip", "Sync asset over remote node.{0}"), DisabledTooltipReason) :
+			FText::Format(LOCTEXT("Initialize_Tooltip_Invalid", "Please ensure the playlist page is using a valid Avalanche Blueprint.{0}"), DisabledTooltipReason);
+
+		const TArray<FAvalanchePage> SelectedTemplatePages = GetSelectedPages(InPlaylist, InPlaylistEditor);
+		
+		MenuBuilder.AddMenuEntry(
+			LOCTEXT("Initialize_Label", "Initialize"),
+			TooltipText,
+			FSlateIcon(FAppStyle::GetAppStyleSetName(), "LevelEditor.ExportAll"),
+			FUIAction(
+				FExecuteAction::CreateSP(this, &FStormSyncAvaPlaylistExtender::HandleInitializeAction, InPlaylist, SelectedTemplatePages),
+				DefaultCanExecuteAction
+			)
+		);
+	}
+
+	constexpr bool bOpenSubMenuOnClick = false;
+
+	// Push action
+	{
+		FText LabelText;
+		FText TooltipText;
+
+		if (!bIsValidSelection)
+		{
+			LabelText = LOCTEXT("PushAssetsMenuEntryInvalid", "Cannot push. Page has no valid asset.");
+			TooltipText = LOCTEXT("PushAssetsMenuEntryTooltipInvalid", "Please ensure the playlist pages are using a valid Avalanche Blueprint.");
+		}
+		else
+		{
+			LabelText = FText::Format(
+				LOCTEXT("PushAssetsAssetName", "Push {0} selected asset{1} to"),
+				FText::AsNumber(SelectedPackageNames.Num()),
+				FText::FromString(SelectedPackageNames.Num() == 1 ? TEXT("") : TEXT("s"))
+			);
+
+			TooltipText = FText::Format(
+				LOCTEXT("PushAssetsMenuEntryTooltip", "Push {0} selected asset{1} (and inner dependencies) to specific remote.\n\nTransfer will only proceed if changes are detected.{2}"),
+				FText::AsNumber(SelectedPackageNames.Num()),
+				FText::FromString(SelectedPackageNames.Num() == 1 ? TEXT("") : TEXT("s")),
+				DisabledTooltipReason
+			);
+		}
+
+		constexpr bool bIsPushing = true;
+		MenuBuilder.AddSubMenu(
+			LabelText,
+			TooltipText,
+			FNewMenuDelegate::CreateRaw(
+				&StormSyncEditor,
+				&FStormSyncEditorModule::BuildPushAssetsMenuSection,
+				SelectedPackageNames,
+				bIsPushing
+			),
+			FUIAction(FExecuteAction(), DefaultCanExecuteAction),
+			NAME_None,
+			EUserInterfaceActionType::Button,
+			bOpenSubMenuOnClick,
+			FSlateIcon(FAppStyle::GetAppStyleSetName(), "LevelEditor.ExportAll")
+		);
+	}
+
+	// Compare action
+	{
+		FText LabelText;
+		FText TooltipText;
+
+		if (!bIsValidSelection)
+		{
+			LabelText = LOCTEXT("CompareAssetsMenuEntryInvalid", "Cannot compare. Page has no valid asset.");
+			TooltipText = LOCTEXT("CompareAssetsMenuEntryTooltipInvalid", "Please ensure the playlist pages are using a valid Avalanche Blueprint.");
+		}
+		else
+		{
+			LabelText = LOCTEXT("CompareAssetsMenuEntry", "Compare Asset(s) with");
+			TooltipText = FText::Format(
+				LOCTEXT("CompareAssetsRemoteMenuEntry", "Compare asset(s) with a specific remote and see if files (and inner dependencies) are either missing or in mismatched state.{0}"),
+				DisabledTooltipReason
+			);
+		}
+		
+		MenuBuilder.AddSubMenu(
+			LabelText,
+			TooltipText,
+			FNewMenuDelegate::CreateRaw(
+				&StormSyncEditor,
+				&FStormSyncEditorModule::BuildCompareWithMenuSection,
+				SelectedPackageNames
+			),
+			FUIAction(FExecuteAction(), DefaultCanExecuteAction),
+			NAME_None,
+			EUserInterfaceActionType::Button,
+			bOpenSubMenuOnClick,
+			FSlateIcon(FAppStyle::GetAppStyleSetName(), "EditorViewport.Visualizers")
+		);
+	}
+	MenuBuilder.EndSection();
+}
+
+void FStormSyncAvaPlaylistExtender::HandleInitializeAction(const UAvalanchePlaylist* InPlaylist, TArray<FAvalanchePage> InSelectedTemplatePages)
+{
+	// Build a list of package names to push grouped by channel name (and remote)
+	// Key is the remote address id, value is the list of package names to synchronize
+	TMap<FString, TArray<FName>> PackageNamesPerChannel;
+	
+	for (const FAvalanchePage& SelectedPage : InSelectedTemplatePages)
+	{
+		const FString AssetName = SelectedPage.GetAvalancheAssetPath(InPlaylist).GetLongPackageName();
+		TArray<FString> ChannelNames = GetChannelNamesForTemplatePage(InPlaylist, SelectedPage);
+
+		// From the list of channel names that match this asset to sync, build the list of server names
+		// Note: there may be more than one server per channel (channel support multiple outputs).
+		
+		STORM_SYNC_AVA_EDITOR_LOG(Display, TEXT("ChannelNames for asset \"%s\" are: %s"), *AssetName, *FString::Join(ChannelNames, TEXT(", ")));
+
+		TArray<FString> RemoteServerNames;
+		for (FString ChannelName : ChannelNames)
+		{
+			const TArray<FString> ServerNames = FStormSyncAvaBridgeUtils::GetServerNamesForChannel(ChannelName);
+			if (ServerNames.IsEmpty())
+			{
+				STORM_SYNC_AVA_EDITOR_LOG(Warning, TEXT("FStormSyncAvaPlaylistExtender::HandleInitializeActions - Unable to determine playback servers for channel \"%s\""), *ChannelName)
+				continue;
+			}
+
+			RemoteServerNames.Append(ServerNames);
+		}
+
+		if (IAvaMediaModule::Get().IsMediaPlaybackClientStarted())
+		{
+			const IAvaMediaPlaybackClient& PlaybackClient = IAvaMediaModule::Get().GetMediaPlaybackClient();
+			for (const FString& ServerName : RemoteServerNames)
+			{
+				// Get address id for storm sync client on playback host
+				FString ClientAddress = PlaybackClient.GetServerUserData(ServerName, UE::StormSync::AvaBridgeCommon::StormSyncClientAddressKey);
+				STORM_SYNC_AVA_EDITOR_LOG(Display, TEXT("ClientAddress on playback server %s is %s"), *ServerName, *ClientAddress);
+
+				if (ClientAddress.IsEmpty())
+				{
+					STORM_SYNC_AVA_EDITOR_LOG(Warning, TEXT("Storm Sync Server Adress id for playback server %s is empty"), *ServerName);
+					continue;
+				}
+				
+				TArray<FName>& PackageNames = PackageNamesPerChannel.FindOrAdd(ClientAddress);
+				PackageNames.Add(FName(*AssetName));
+			}
+		}
+	}
+
+	for (const TPair<FString, TArray<FName>>& NamesPerChannel : PackageNamesPerChannel)
+	{
+		FString AddressId = NamesPerChannel.Key;
+		TArray<FName> PackageNames = NamesPerChannel.Value;
+		
+		STORM_SYNC_AVA_EDITOR_LOG(Display, TEXT("Names per channel - Channel: %s"), *AddressId)
+		for (const FName& PackageName : PackageNames)
+		{
+			STORM_SYNC_AVA_EDITOR_LOG(Display, TEXT("\t PackageName: %s"), *PackageName.ToString())
+		}
+
+		PushPackagesToRemote(AddressId, PackageNames);
+	}
+}
+
+TSharedRef<FExtender> FStormSyncAvaPlaylistExtender::AddToolbarExtender(const TSharedRef<FUICommandList> InCommandList, const TArray<UObject*> ContextSensitiveObjects)
+{
+	STORM_SYNC_AVA_EDITOR_LOG(Display, TEXT("FStormSyncAvaPlaylistExtender::AddToolbarExtender - Adding in toolbar extensions"))
+	TSharedRef<FExtender> Extender(new FExtender());
+
+	const UAvalanchePlaylist* Playlist = ContextSensitiveObjects.IsValidIndex(0) ? Cast<UAvalanchePlaylist>(ContextSensitiveObjects[0]) : nullptr;
+	if (!Playlist)
+	{
+		return Extender;
+	}
+
+	const TSharedPtr<IToolkit> AssetEditor = FToolkitManager::Get().FindEditorForAsset(Playlist);
+	if (!AssetEditor.IsValid())
+	{
+		return Extender;
+	}
+
+	const TSharedPtr<FAvaPlaylistEditor> PlaylistEditor = StaticCastSharedPtr<FAvaPlaylistEditor>(AssetEditor);
+	if (!PlaylistEditor.IsValid())
+	{
+		return Extender;
+	}
+
+	Extender->AddToolBarExtension(
+		"Pages",
+		EExtensionHook::After,
+		InCommandList,
+		FToolBarExtensionDelegate::CreateSP(this, &FStormSyncAvaPlaylistExtender::FillToolbar, TWeakPtr<FAvaPlaylistEditor>(PlaylistEditor))
+	);
+	
+	STORM_SYNC_AVA_EDITOR_LOG(Display, TEXT("FStormSyncAvaPlaylistExtender::AddToolbarExtender - Playlist: %s"), *GetNameSafe(Playlist))
+	
+	return Extender;
+}
+
+void FStormSyncAvaPlaylistExtender::FillToolbar(FToolBarBuilder& ToolbarBuilder, TWeakPtr<FAvaPlaylistEditor> InPlaylistEditor)
+{
+	ToolbarBuilder.BeginSection(TEXT("StormSync"));
+	{
+		ToolbarBuilder.AddComboButton(
+			FUIAction(),
+			FOnGetContent::CreateSP(this, &FStormSyncAvaPlaylistExtender::GenerateToolbarMenu, InPlaylistEditor),
+			LOCTEXT("ToolbarLabel", "Synchronize Actions"),
+			LOCTEXT("ToolbarToolTip", "Synchronize the playlist assets over the network"),
+			FSlateIcon(FAppStyle::GetAppStyleSetName(), "LevelEditor.Recompile"),
+			false
+		);
+	}
+	ToolbarBuilder.EndSection();
+}
+
+TSharedRef<SWidget> FStormSyncAvaPlaylistExtender::GenerateToolbarMenu(TWeakPtr<FAvaPlaylistEditor> InPlaylistEditor)
+{
+	TArray<FName> PackageNames;
+
+	const TSharedPtr<FAvaPlaylistEditor> PlaylistEditor = InPlaylistEditor.Pin();
+
+	if (PlaylistEditor.IsValid() && PlaylistEditor->IsPlaylistValid())
+	{
+		const UAvalanchePlaylist* Playlist = PlaylistEditor->GetPlaylist(); 
+
+		// Gather the list of all package names from Avalanche Blueprints in this Playlist pages
+		const FAvalanchePageCollection& AvalanchePageCollection = Playlist->GetTemplatePages();
+		for (const FAvalanchePage& AvalanchePage : AvalanchePageCollection.Pages)
+		{
+			if (FString PackageName = AvalanchePage.GetAvalancheAssetPath(Playlist).GetLongPackageName(); !PackageName.IsEmpty())
+			{
+				PackageNames.AddUnique(FName(*PackageName));
+			}
+		}
+	}
+	
+	FMenuBuilder MenuBuilder(true, nullptr);
+
+	MenuBuilder.BeginSection(TEXT("StormSyncActions"));
+	
+	if (FModuleManager::Get().IsModuleLoaded(TEXT("StormSyncEditor")))
+	{
+		FStormSyncEditorModule& StormSyncEditor = FStormSyncEditorModule::Get();
+		
+		const int32 PackagesCount = PackageNames.Num();
+
+		TArray<FString> AssetList;
+		for (const FName PackageName : PackageNames)
+		{
+			AssetList.Add(FString::Printf(TEXT("- %s"), *PackageName.ToString()));
+		}
+
+		FText LabelText = FText::Format(LOCTEXT("PushAssetsToolbarMenuEntry", "Push {0} asset(s) to"), FText::AsNumber(PackagesCount));
+		FText TooltipText = FText::Format(LOCTEXT(
+			"PushAssetsToolbarMenuEntryTooltip",
+			"Push {0} asset(s) (and inner dependencies) to specific remote.\n\nTransfer will only proceed if changes are detected.\n\n{1}"
+		), FText::AsNumber(PackagesCount), FText::FromString(FString::Join(AssetList, LINE_TERMINATOR)));
+		
+		const bool bIsPushEnabled = !PackageNames.IsEmpty();
+		if (!bIsPushEnabled)
+		{
+			LabelText = LOCTEXT("PushAssetsToolbarMenuEntryInvalid", "Cannot push. Playlist pages have no valid assets.");
+			TooltipText = LOCTEXT("PushAssetsToolbarMenuEntryTooltipInvalid", "Please ensure the playlist pages are using a valid Avalanche Blueprint");
+		}
+
+		constexpr bool bIsPushing = true;
+		constexpr bool bOpenSubMenuOnClick = false;
+		MenuBuilder.AddSubMenu(
+			LabelText,
+			TooltipText,
+			FNewMenuDelegate::CreateRaw(
+				&StormSyncEditor,
+				&FStormSyncEditorModule::BuildPushAssetsMenuSection,
+				PackageNames,
+				bIsPushing
+			),
+			FUIAction(
+				FExecuteAction(),
+				FCanExecuteAction::CreateLambda([bIsPushEnabled]() { return bIsPushEnabled; })
+			),
+			NAME_None,
+			EUserInterfaceActionType::Button,
+			bOpenSubMenuOnClick,
+			FSlateIcon(FAppStyle::GetAppStyleSetName(), "LevelEditor.ExportAll")
+		);
+
+		MenuBuilder.AddSubMenu(
+			FText::Format(LOCTEXT("CompareAssetsToolbarMenuEntry", "Compare {0} asset(s) with"), FText::AsNumber(PackagesCount)),
+			LOCTEXT("CompareAssetsToolbarMenuEntryTooltip", "Compare asset(s) with a specific remote and see if files (and inner dependencies) are either missing or in mismatched state."),
+			FNewMenuDelegate::CreateRaw(
+				&StormSyncEditor,
+				&FStormSyncEditorModule::BuildCompareWithMenuSection,
+				PackageNames
+			),
+			FUIAction(
+				FExecuteAction(),
+				FCanExecuteAction::CreateLambda([bIsPushEnabled]() { return bIsPushEnabled; })
+			),
+			NAME_None,
+			EUserInterfaceActionType::Button,
+			bOpenSubMenuOnClick,
+			FSlateIcon(FAppStyle::GetAppStyleSetName(), "EditorViewport.Visualizers")
+		);
+	}
+	MenuBuilder.EndSection();
+
+	return MenuBuilder.MakeWidget();
+}
+
+void FStormSyncAvaPlaylistExtender::PushPackagesToRemote(const FString& RemoteAddressId, const TArray<FName>& InPackageNames)
+{
+	STORM_SYNC_AVA_EDITOR_LOG(Display, TEXT("FStormSyncAvaPlaylistExtender::PushPackagesToRemote - InSelectedPage: (InPackageNames: %d, RemoteAddressId: %s)"), InPackageNames.Num(), *RemoteAddressId)
+
+	FMessageAddress RemoteMessageAddress;
+	if (!FMessageAddress::Parse(RemoteAddressId, RemoteMessageAddress))
+	{
+		STORM_SYNC_AVA_EDITOR_LOG(Error, TEXT("Unable to parse %s into a Message Address"), *RemoteAddressId);
+		return;
+	}
+
+	// Note: We sync with a dummy package descriptor, next iterations could add in there an additional UI step.
+	// Something that could be handled with a bit more UI integration, like some kind of popup window or wizard.
+	const FStormSyncPackageDescriptor PackageDescriptor;
+
+	const FOnStormSyncPushComplete Delegate = FOnStormSyncPushComplete::CreateLambda([](const TSharedPtr<FStormSyncTransportPushResponse>& Response)
+	{
+		check(Response.IsValid())
+		STORM_SYNC_AVA_EDITOR_LOG(Display, TEXT("FStormSyncAvaPlaylistExtender::PushPackagesToRemote - Got a response: %s"), *Response->ToString());
+		UStormSyncNotificationSubsystem::Get().HandlePushResponse(Response);
+	});
+	
+	IStormSyncTransportClientModule::Get().PushPackages(PackageDescriptor, InPackageNames, RemoteMessageAddress, Delegate);
+}
+
+TArray<FName> FStormSyncAvaPlaylistExtender::GetSelectedPackagesNames(const UAvalanchePlaylist* InPlaylist, const TWeakPtr<FAvaPlaylistEditor>& InPlaylistEditor)
+{
+	check(InPlaylist);
+	
+	TArray<FName> Result;
+
+	const TArray<FAvalanchePage> SelectedPages = GetSelectedPages(InPlaylist, InPlaylistEditor);
+	Algo::Transform(SelectedPages, Result, [InPlaylist](const FAvalanchePage& Page)
+	{
+		return Page.GetAvalancheAssetPath(InPlaylist).GetLongPackageFName();
+	});
+
+	return Result;
+}
+
+TArray<FAvalanchePage> FStormSyncAvaPlaylistExtender::GetSelectedPages(const UAvalanchePlaylist* InPlaylist, const TWeakPtr<FAvaPlaylistEditor>& InPlaylistEditor)
+{
+	check(InPlaylist);
+	
+	TArray<FAvalanchePage> Result;
+	const TSharedPtr<FAvaPlaylistEditor> PlaylistEditor = InPlaylistEditor.Pin();
+	if (!PlaylistEditor.IsValid())
+	{
+		STORM_SYNC_AVA_EDITOR_LOG(Display, TEXT("FStormSyncAvaPlaylistExtender::GetSelectedPages - Invalid shared ptr from weak ptr delegate param"))
+		return Result;
+	}
+
+	const TConstArrayView<int32> SelectedPageIds = PlaylistEditor->GetSelectedPagesOnFocusedWidget();
+	for (const int32 SelectedPageId : SelectedPageIds)
+	{
+		FAvalanchePage AvalanchePage = InPlaylist->GetPage(SelectedPageId);
+		if (!AvalanchePage.IsValidPage())
+		{
+			continue;
+		}
+
+		FSoftObjectPath SoftAvalancheAssetPath = AvalanchePage.GetAvalancheAssetPath(InPlaylist);
+		FString LongPackageName = SoftAvalancheAssetPath.GetLongPackageName();
+		FString AssetName = SoftAvalancheAssetPath.GetAssetName();
+
+		if (!LongPackageName.IsEmpty())
+		{
+			Result.Add(AvalanchePage);
+		}
+	}
+	
+	return Result;
+}
+
+TArray<FString> FStormSyncAvaPlaylistExtender::GetChannelNamesForTemplatePage(const UAvalanchePlaylist* InPlaylist, const FAvalanchePage& InTemplatePage)
+{
+	check(InPlaylist);
+
+	TArray<FString> ChannelNames;
+
+	// Here, we try to determine the list of channels to consider for a sync operation, from the selected template page,
+	// with instanced pages that are using the selected package name (Avalanche Blueprint)
+	const FString PackageName = InTemplatePage.GetAvalancheAssetPath(InPlaylist).GetLongPackageName();
+
+	// Build the list of instanced pages matching the asset name we want to sync
+	TArray<FAvalanchePage> Pages = InPlaylist->GetInstancedPages().Pages.FilterByPredicate([InPlaylist, PackageName](const FAvalanchePage& Page)
+	{
+		return Page.GetAvalancheAssetPath(InPlaylist).GetLongPackageName() == PackageName;
+	});
+
+	// From there, build a unique list of channel outputs
+	for (const FAvalanchePage& AvalanchePage : Pages)
+	{
+		ChannelNames.AddUnique(AvalanchePage.GetChannelName().ToString());
+	}
+	
+	return ChannelNames;
+}
+
+bool FStormSyncAvaPlaylistExtender::GetContextMenuSelectionInfos(const UAvalanchePlaylist* InPlaylist, const TWeakPtr<FAvaPlaylistEditor>& InPlaylistEditor, bool& bOutIsValidSelection, FText& OutDisabledReasonTooltip, TArray<FName>& OutSelectedPackageNames, FCanExecuteAction& OutCanExecuteAction)
+{
+	if (!FModuleManager::Get().IsModuleLoaded(TEXT("StormSyncEditor")))
+	{
+		STORM_SYNC_AVA_EDITOR_LOG(Error, TEXT("FStormSyncAvaPlaylistExtender::GetContextMenuSelectionInfos - StormSyncEditor module is not loaded. Is StormSync plugin enabled ?"))
+		return false;
+	}
+	
+	const FStormSyncEditorModule& StormSyncEditor = FStormSyncEditorModule::Get();
+
+	TArray<FName> SelectedPackagesNames = GetSelectedPackagesNames(InPlaylist, InPlaylistEditor);
+	bool bIsValidSelection = !SelectedPackagesNames.IsEmpty();
+
+	// Figure out if selection is containing dirty (unsaved) assets
+	FText DisabledTooltipReason;
+	const TArray<FAssetData> DirtyAssets = StormSyncEditor.GetDirtyAssets(SelectedPackagesNames, DisabledTooltipReason);
+	
+	bool bContainsDirtyAssets = !DirtyAssets.IsEmpty();
+	FCanExecuteAction DefaultCanExecuteAction = FCanExecuteAction::CreateLambda([bIsValidSelection, bContainsDirtyAssets]() { return bIsValidSelection && !bContainsDirtyAssets; });
+
+	bOutIsValidSelection = bIsValidSelection;
+	OutDisabledReasonTooltip = DisabledTooltipReason;
+	OutSelectedPackageNames = MoveTemp(SelectedPackagesNames);
+	OutCanExecuteAction = MoveTemp(DefaultCanExecuteAction);
+	
+	return true;
+}
+
+#undef LOCTEXT_NAMESPACE
diff --git a/Engine/Plugins/Experimental/StormSyncAvalancheBridge/Source/StormSyncAvaBridgeEditor/Private/StormSyncAvaPlaylistExtender.h b/Engine/Plugins/Experimental/StormSyncAvalancheBridge/Source/StormSyncAvaBridgeEditor/Private/StormSyncAvaPlaylistExtender.h
new file mode 100644
--- /dev/null
+++ b/Engine/Plugins/Experimental/StormSyncAvalancheBridge/Source/StormSyncAvaBridgeEditor/Private/StormSyncAvaPlaylistExtender.h
@@ -0,0 +1,100 @@
+// Copyright Epic Games, Inc. All Rights Reserved.
+
+#pragma once
+
+#include "CoreMinimal.h"
+#include "MessageEndpoint.h"
+#include "Toolkits/AssetEditorToolkit.h"
+
+class FAvaPlaylistEditor;
+class FExtensibilityManager;
+class UAvalanchePlaylist;
+struct FAvalanchePage;
+
+/** This class handles Avalanche Playlist assets UI extensions for context menu / toolbars */
+class FStormSyncAvaPlaylistExtender : public TSharedFromThis<FStormSyncAvaPlaylistExtender>
+{
+public:
+	FStormSyncAvaPlaylistExtender();
+	~FStormSyncAvaPlaylistExtender();
+
+private:
+	/** Holds the messaging endpoint. */
+	TSharedPtr<FMessageEndpoint, ESPMode::ThreadSafe> MessageEndpoint;
+
+	/** Name of the extension point we're providing extension next to*/
+	static constexpr const TCHAR* AvalancheExtensionHook = TEXT("PageListOperations");
+	
+	/** Menu extender for avalanche playlist editor context menu */
+	FDelegateHandle MenuExtenderHandle;
+	
+	/** Toolbar extender for avalanche playlist editor */
+	FDelegateHandle ToolbarExtenderHandle;
+
+	/** Register Context menu / toolbar extensions when editor is fully loaded */
+	void OnPostEngineInit();
+	
+	/** Startup module handler for UI extensions */
+	void RegisterMenuExtensions();
+
+	/** Startup module handler to register an UI extensions from one of Avalanche exposed extensibility hooks */
+	static FDelegateHandle RegisterExtension(const TSharedPtr<FExtensibilityManager> InExtensibilityManager, const FAssetEditorExtender& InExtenderDelegate);
+	
+	/** Shutdown module handler to unregister any UI extensions added here */
+	static void UnregisterExtension(const TSharedPtr<FExtensibilityManager> InExtensibilityManager, const FDelegateHandle& InHandleToRemove);
+	
+	/** Gets the extender to use for playlist context sensitive menus */
+	TSharedRef<FExtender> AddMenuExtender(const TSharedRef<FUICommandList> InCommandList, const TArray<UObject*> ContextSensitiveObjects);
+
+	/* UI Menu Extension handler for template panel */
+	void CreateTemplateContextMenu(FMenuBuilder& MenuBuilder, const UAvalanchePlaylist* InPlaylist, TWeakPtr<FAvaPlaylistEditor> InPlaylistEditor);
+
+	/** Context menu handler for initialize action */
+	void HandleInitializeAction(const UAvalanchePlaylist* InPlaylist, TArray<FAvalanchePage> InSelectedTemplatePages);
+	
+	/** Gets the extender to use for playlist context sensitive menus */
+	TSharedRef<FExtender> AddToolbarExtender(const TSharedRef<FUICommandList> InCommandList, const TArray<UObject*> ContextSensitiveObjects);
+	
+	/** Construct toolbar widgets for playlist sync actions */
+	void FillToolbar(FToolBarBuilder& ToolbarBuilder, TWeakPtr<FAvaPlaylistEditor> InPlaylistEditor);
+
+	/** Creates widget for toolbar content */
+	TSharedRef<SWidget> GenerateToolbarMenu(TWeakPtr<FAvaPlaylistEditor> InPlaylistEditor);
+
+	/** Storm sync push a list of package names to a remote address ID */
+	static void PushPackagesToRemote(const FString& RemoteAddressId, const TArray<FName>& InPackageNames);
+
+	/** Gather a list of package names from currently selected pages in editor */
+	static TArray<FName> GetSelectedPackagesNames(const UAvalanchePlaylist* InPlaylist, const TWeakPtr<FAvaPlaylistEditor>& InPlaylistEditor);
+	
+	/**
+	 * Returns the list of currently selected pages in editor
+	 *
+	 * It will only return selection if the page is referencing a valid Avalanche Blueprint (not "None" path)
+	 */
+	static TArray<FAvalanchePage> GetSelectedPages(const UAvalanchePlaylist* InPlaylist, const TWeakPtr<FAvaPlaylistEditor>& InPlaylistEditor);
+
+	/** Returns a unique list of channel names, gather from instanced pages, that are matching the passed in template page selection and selected asset name */
+	static TArray<FString> GetChannelNamesForTemplatePage(const UAvalanchePlaylist* InPlaylist, const FAvalanchePage& InTemplatePage);
+
+	/**
+	 * Provides a way for context menu builders to gather information about current selection.
+	 *
+	 * @param InPlaylist The playlist UObject we're working with
+	 * @param InPlaylistEditor The playlist asset editor we're working with
+	 * @param bOutIsValidSelection Output indicating whether selection is valid - eg. current selection returns a list of valid package names
+	 * @param OutDisabledReasonTooltip Output holding the tooltip to use in case the action is disabled - eg. in case the selection contains unsaved dirty assets
+	 * @param OutSelectedPackageNames Output list of package names of the selection, as returned by GetSelectedPackagesNames()
+	 * @param OutCanExecuteAction Output delegate to use as a default CanExecuteAction - eg. disabling the action if selection is invalid
+	 *
+	 * @return false if StormSyncEditor module is not available, and we failed to determine selection state
+	 */
+	static bool GetContextMenuSelectionInfos(
+		const UAvalanchePlaylist* InPlaylist,
+		const TWeakPtr<FAvaPlaylistEditor>& InPlaylistEditor,
+		bool& bOutIsValidSelection,
+		FText& OutDisabledReasonTooltip,
+		TArray<FName>& OutSelectedPackageNames,
+		FCanExecuteAction& OutCanExecuteAction
+	);
+};
diff --git a/Engine/Plugins/Experimental/StormSyncAvalancheBridge/Source/StormSyncAvaBridgeEditor/Public/StormSyncAvaBridgeEditor.h b/Engine/Plugins/Experimental/StormSyncAvalancheBridge/Source/StormSyncAvaBridgeEditor/Public/StormSyncAvaBridgeEditor.h
new file mode 100644
--- /dev/null
+++ b/Engine/Plugins/Experimental/StormSyncAvalancheBridge/Source/StormSyncAvaBridgeEditor/Public/StormSyncAvaBridgeEditor.h
@@ -0,0 +1,22 @@
+// Copyright Epic Games, Inc. All Rights Reserved.
+
+#pragma once
+
+#include "CoreMinimal.h"
+#include "Modules/ModuleInterface.h"
+#include "Modules/ModuleManager.h"
+
+class FStormSyncAvaPlaylistExtender;
+
+class FStormSyncAvaBridgeEditorModule : public IModuleInterface
+{
+public:
+	//~ Begin IModuleInterface interface
+	virtual void StartupModule() override;
+	virtual void ShutdownModule() override;
+	//~ End IModuleInterface interface
+
+private:
+	/** UI extender for Storm Sync view and commands. */
+	TSharedPtr<FStormSyncAvaPlaylistExtender> PlaylistExtender;
+};
diff --git a/Engine/Plugins/Experimental/StormSyncAvalancheBridge/Source/StormSyncAvaBridgeEditor/Public/StormSyncAvaBridgeEditorLog.h b/Engine/Plugins/Experimental/StormSyncAvalancheBridge/Source/StormSyncAvaBridgeEditor/Public/StormSyncAvaBridgeEditorLog.h
new file mode 100644
--- /dev/null
+++ b/Engine/Plugins/Experimental/StormSyncAvalancheBridge/Source/StormSyncAvaBridgeEditor/Public/StormSyncAvaBridgeEditorLog.h
@@ -0,0 +1,12 @@
+// Copyright Epic Games, Inc. All Rights Reserved.
+
+#pragma once
+
+#include "CoreMinimal.h"
+
+STORMSYNCAVABRIDGEEDITOR_API DECLARE_LOG_CATEGORY_EXTERN(LogStormSyncAvaBridgeEditor, Display, All);
+
+#define STORM_SYNC_AVA_EDITOR_LOG(Verbosity, Format, ...) \
+{ \
+    UE_LOG(LogStormSyncAvaBridgeEditor, Verbosity, Format, ##__VA_ARGS__); \
+}
diff --git a/Engine/Plugins/Experimental/StormSyncAvalancheBridge/Source/StormSyncAvaBridgeEditor/StormSyncAvaBridgeEditor.Build.cs b/Engine/Plugins/Experimental/StormSyncAvalancheBridge/Source/StormSyncAvaBridgeEditor/StormSyncAvaBridgeEditor.Build.cs
new file mode 100644
--- /dev/null
+++ b/Engine/Plugins/Experimental/StormSyncAvalancheBridge/Source/StormSyncAvaBridgeEditor/StormSyncAvaBridgeEditor.Build.cs
@@ -0,0 +1,37 @@
+// Copyright Epic Games, Inc. All Rights Reserved.
+
+using UnrealBuildTool;
+
+public class StormSyncAvaBridgeEditor : ModuleRules
+{
+	public StormSyncAvaBridgeEditor(ReadOnlyTargetRules Target) : base(Target)
+	{
+		PublicDependencyModuleNames.AddRange(
+			new string[]
+			{
+				"Core"
+			}
+		);
+
+		PrivateDependencyModuleNames.AddRange(
+			new string[]
+			{
+				"AvalancheMedia",
+				"AvalancheMediaEditor",
+				"CoreUObject",
+				"EditorFramework",
+				"Engine",
+				"Slate",
+				"SlateCore",
+				"StormSyncAvaBridge",
+				"StormSyncCore",
+				"StormSyncEditor",
+				"StormSyncTransportCore",
+				"StormSyncTransportClient",
+				"StormSyncTransportServer",
+			}
+		);
+
+		ShortName = "StScAvBridEd";
+	}
+}
diff --git a/Engine/Plugins/Experimental/StormSyncAvalancheBridge/StormSyncAvalancheBridge.uplugin b/Engine/Plugins/Experimental/StormSyncAvalancheBridge/StormSyncAvalancheBridge.uplugin
new file mode 100644
--- /dev/null
+++ b/Engine/Plugins/Experimental/StormSyncAvalancheBridge/StormSyncAvalancheBridge.uplugin
@@ -0,0 +1,38 @@
+{
+	"FileVersion": 3,
+	"Version": 1,
+	"VersionName": "0.16.0",
+	"FriendlyName": "Storm Sync Motion Design Bridge",
+	"Description": "Plugin bridge between Motion Design Plugin and Storm Sync to provide in-editor integration to synchronize assets",
+	"Category": "Experimental",
+	"CreatedBy": "Epic Games, Inc.",
+	"CreatedByURL": "",
+	"DocsURL": "",
+	"MarketplaceURL": "",
+	"SupportURL": "",
+	"CanContainContent": true,
+	"IsBetaVersion": true,
+	"Installed": false,
+	"Modules": [
+		{
+			"Name": "StormSyncAvaBridgeEditor",
+			"Type": "Editor",
+			"LoadingPhase": "Default"
+		},
+		{
+			"Name": "StormSyncAvaBridge",
+			"Type": "Runtime",
+			"LoadingPhase": "Default"
+		}
+	],
+	"Plugins": [
+		{
+			"Name": "Avalanche",
+			"Enabled": true
+		},
+		{
+			"Name": "StormSync",
+			"Enabled": true
+		}
+	]
+}
