Motion Design

diff --git a/Engine/Plugins/Experimental/Avalanche/Source/AvalancheMediaEditor/Private/Rundown/MediaOutputEditorUtils/AvaRundownOutputEditorUtils.h b/Engine/Plugins/Experimental/Avalanche/Source/AvalancheMediaEditor/Private/Rundown/MediaOutputEditorUtils/AvaRundownOutputEditorUtils.h
new file mode 100644
--- /dev/null
+++ b/Engine/Plugins/Experimental/Avalanche/Source/AvalancheMediaEditor/Private/Rundown/MediaOutputEditorUtils/AvaRundownOutputEditorUtils.h
@@ -0,0 +1,40 @@
+// Copyright Epic Games, Inc. All Rights Reserved.
+
+#pragma once
+
+#include "MediaOutput.h"
+#include "Templates/SharedPointer.h"
+
+DECLARE_LOG_CATEGORY_EXTERN(LogAvaRundownEditorOutputUtils, Warning, All);
+
+class FJsonObject;
+
+/*
+ *	Class designed to deal with Serialization and Editing
+ *	of Media Output device data coming from Rundown Server
+ *
+ *	It is meant to give full control and provide all information to Rundown server
+ *	in a form that is usable outside of unreal engine
+ */
+class FAvaRundownOutputEditorUtils
+{
+public:
+	FAvaRundownOutputEditorUtils();
+	virtual ~FAvaRundownOutputEditorUtils();
+
+	static FString SerializeMediaOutput(const UMediaOutput* InMediaOutput);
+	static void EditMediaOutput(UMediaOutput* InMediaOutput, const FString& InDeviceData);
+
+private:
+	// Serialization helper functions
+	static TSharedPtr<FJsonObject> ParsePropertyInfo(TFieldIterator<FProperty> InProperty, const void* InOwnerObject);
+	static TSharedPtr<FJsonObject> ParseElementaryPropertyInfo(TFieldIterator<FProperty> InProperty, const void* InOwnerObject);
+	static TSharedPtr<FJsonObject> ParseEnumPropertyInfo(TFieldIterator<FProperty> InProperty, const void* InOwnerObject);
+	static TSharedPtr<FJsonObject> ParseStructPropertyInfo(TFieldIterator<FProperty> InProperty, const void* InOwnerObject);
+
+	// Helper functions for editing Media Output Devices
+	static void SetProperty(void* InOwnerObject, const TSharedPtr<FJsonObject>& InPropertyObject, const TFieldIterator<FProperty>& InProperty);
+	static void SetElementaryProperty(void* InOwnerObject, const TSharedPtr<FJsonObject>& InPropertyObject, TFieldIterator<FProperty> InProperty);
+	static void SetEnumProperty(void* InOwnerObject, const TSharedPtr<FJsonObject>& InPropertyObject, TFieldIterator<FProperty> InProperty);
+	static void SetStructProperties(void* InOwnerObject, const TSharedPtr<FJsonObject>& InPropertyObject, TFieldIterator<FProperty> InProperty);
+};
diff --git a/Engine/Plugins/Experimental/Avalanche/Source/AvalancheMediaEditor/Private/Rundown/TabFactories/AvaRundownSubListDocumentTabFactory.cpp b/Engine/Plugins/Experimental/Avalanche/Source/AvalancheMediaEditor/Private/Rundown/TabFactories/AvaRundownSubListDocumentTabFactory.cpp
new file mode 100644
--- /dev/null
+++ b/Engine/Plugins/Experimental/Avalanche/Source/AvalancheMediaEditor/Private/Rundown/TabFactories/AvaRundownSubListDocumentTabFactory.cpp
@@ -0,0 +1,65 @@
+// Copyright Epic Games, Inc. All Rights Reserved.
+
+#include "AvaRundownSubListDocumentTabFactory.h"
+
+#include "Rundown/AvaRundownEditor.h"
+#include "Rundown/Pages/Slate/SAvaRundownInstancedPageList.h"
+#include "Widgets/Docking/SDockTab.h"
+
+#define LOCTEXT_NAMESPACE "AvaRundownSubListDocumentTabFactory"
+
+const FName FAvaRundownSubListDocumentTabFactory::FactoryId = "AvaSubListTabFactory";
+const FString FAvaRundownSubListDocumentTabFactory::BaseTabName(TEXT("AvaSubListDocument"));
+
+FName FAvaRundownSubListDocumentTabFactory::GetTabId(int32 InSubListIndex)
+{
+	return FName(BaseTabName + "_" + FString::FromInt(InSubListIndex));
+}
+
+FAvaRundownSubListDocumentTabFactory::FAvaRundownSubListDocumentTabFactory(const TSharedPtr<FAvaRundownEditor>& InRundownEditor)
+	: FDocumentTabFactory(FactoryId, InRundownEditor)
+	, RundownEditorWeak(InRundownEditor)
+	, SubListIndex(UAvaRundown::InstancePageList.SubListIndex)
+{
+	TabIcon = FSlateIcon(FAppStyle::GetAppStyleSetName(), "PlacementBrowser.Icons.All");
+}
+
+TSharedRef<SWidget> FAvaRundownSubListDocumentTabFactory::CreateTabBody(const FWorkflowTabSpawnInfo& InInfo) const
+{
+	if (SubListIndex <= UAvaRundown::InstancePageList.SubListIndex)
+	{
+		return SNullWidget::NullWidget;
+	}
+
+	return SNew(SAvaRundownInstancedPageList, RundownEditorWeak.Pin(), UAvaRundown::CreateSubListReference(SubListIndex));
+}
+
+TSharedRef<SDockTab> FAvaRundownSubListDocumentTabFactory::SpawnSubListTab(const FWorkflowTabSpawnInfo& InInfo, int32 InSubListIndex)
+{
+	SubListIndex = InSubListIndex;
+	TabIdentifier = GetTabId(InSubListIndex);
+	TabLabel = FText::Format(LOCTEXT("RundownSubListDocument_TabLabel", "Page View {0}"), FText::AsNumber(SubListIndex + 1));
+	ViewMenuDescription = FText::Format(LOCTEXT("RundownSubListDocument_ViewMenu_Desc", "Page View {0}"), FText::AsNumber(SubListIndex + 1));
+	ViewMenuTooltip = FText::Format(LOCTEXT("RundownSubListDocument_ViewMenu_ToolTip", "Page View {0}"), FText::AsNumber(SubListIndex + 1));
+
+	TSharedRef<SDockTab> NewTab = SpawnTab(InInfo);
+
+	TSharedRef<SAvaRundownInstancedPageList> PageList = StaticCastSharedRef<SAvaRundownInstancedPageList>(NewTab->GetContent());
+	NewTab->SetOnTabActivated(SDockTab::FOnTabActivatedCallback::CreateSP(PageList, &SAvaRundownInstancedPageList::OnTabActivated));
+	PageList->SetMyTab(NewTab);
+
+	if (TSharedPtr<FAvaRundownEditor> RundownEditor = PageList->GetRundownEditor())
+	{
+		UAvaRundown* Rundown = RundownEditor->GetRundown();
+
+		if (IsValid(Rundown) && Rundown->IsValidSubListIndex(InSubListIndex)
+			&& !Rundown->GetSubList(InSubListIndex).Name.IsEmpty())
+		{
+			NewTab->SetLabel(Rundown->GetSubList(InSubListIndex).Name);
+		}
+	}
+
+	return NewTab;
+}
+
+#undef LOCTEXT_NAMESPACE
diff --git a/Engine/Plugins/Experimental/Avalanche/Source/AvalancheMediaEditor/Private/Rundown/TabFactories/AvaRundownSubListDocumentTabFactory.h b/Engine/Plugins/Experimental/Avalanche/Source/AvalancheMediaEditor/Private/Rundown/TabFactories/AvaRundownSubListDocumentTabFactory.h
new file mode 100644
--- /dev/null
+++ b/Engine/Plugins/Experimental/Avalanche/Source/AvalancheMediaEditor/Private/Rundown/TabFactories/AvaRundownSubListDocumentTabFactory.h
@@ -0,0 +1,28 @@
+// Copyright Epic Games, Inc. All Rights Reserved.
+
+#pragma once
+
+#include "WorkflowOrientedApp/WorkflowTabFactory.h"
+
+class FAvaRundownEditor;
+
+/** Base class for all Tab Factories in Ava SubListDocument Editor */
+class FAvaRundownSubListDocumentTabFactory : public FDocumentTabFactory
+{
+public:
+	static const FName FactoryId;
+	static const FString BaseTabName;
+
+	static FName GetTabId(int32 InSubListIndex);
+
+	FAvaRundownSubListDocumentTabFactory(const TSharedPtr<FAvaRundownEditor>& InSubListDocumentEditor);
+
+	virtual TSharedRef<SWidget> CreateTabBody(const FWorkflowTabSpawnInfo& InInfo) const override;
+
+	TSharedRef<SDockTab> SpawnSubListTab(const FWorkflowTabSpawnInfo& InInfo, int32 InSubListIndex);
+
+protected:
+	TWeakPtr<FAvaRundownEditor> RundownEditorWeak;
+	int32 SubListIndex;
+};
+
diff --git a/Engine/Plugins/Experimental/Avalanche/Source/AvalancheMediaEditor/Public/Broadcast/AvaBroadcastEditor.h b/Engine/Plugins/Experimental/Avalanche/Source/AvalancheMediaEditor/Public/Broadcast/AvaBroadcastEditor.h
--- a/Engine/Plugins/Experimental/Avalanche/Source/AvalancheMediaEditor/Public/Broadcast/AvaBroadcastEditor.h
+++ b/Engine/Plugins/Experimental/Avalanche/Source/AvalancheMediaEditor/Public/Broadcast/AvaBroadcastEditor.h
@@ -4,11 +4,10 @@
 
 #include "AvaMediaDefines.h"
 #include "CoreMinimal.h"
-#include "MediaOutput.h"
 #include "WorkflowOrientedApp/WorkflowCentricApplication.h"
 
-class FAvaOutputTileItem;
-class UAvalancheBroadcast;
+class FAvaBroadcastOutputTileItem;
+class UAvaBroadcast;
 class UMediaOutput;
 
 class FAvaBroadcastEditor : public FWorkflowCentricApplication
@@ -23,15 +22,15 @@ public:
 
 	static void OpenBroadcastEditor();
 
-	void SelectOutputTile(const TSharedPtr<FAvaOutputTileItem>& InOutputTile);
-	TSharedPtr<FAvaOutputTileItem> GetSelectedOutputTile() const;
+	void SelectOutputTile(const TSharedPtr<FAvaBroadcastOutputTileItem>& InOutputTile);
+	TSharedPtr<FAvaBroadcastOutputTileItem> GetSelectedOutputTile() const;
 	
-	DECLARE_MULTICAST_DELEGATE_OneParam(FOnOutputItemSelectionChanged, const TSharedPtr<FAvaOutputTileItem>&);
+	DECLARE_MULTICAST_DELEGATE_OneParam(FOnOutputItemSelectionChanged, const TSharedPtr<FAvaBroadcastOutputTileItem>&);
 	FOnOutputItemSelectionChanged OnOutputTileSelectionChanged;
 	
 protected:
 	
-	void InitBroadcastEditor(UAvalancheBroadcast* InBroadcast);
+	void InitBroadcastEditor(UAvaBroadcast* InBroadcast);
 
 	void OnBroadcastChanged(EAvaBroadcastChange ChangedEvent);
 	void OnAvaMediaSettingsChanged(UObject*, FPropertyChangedEvent&);
@@ -57,7 +56,7 @@ protected:
 
 public:
 	
-	UAvalancheBroadcast* GetBroadcastObject() const;
+	UAvaBroadcast* GetBroadcastObject() const;
 
 	void ExtendToolBar(TSharedPtr<FExtender> Extender);
 	void FillPlayToolBar(FToolBarBuilder& ToolBarBuilder);
@@ -82,7 +81,7 @@ protected:
 	
 	static TSharedPtr<FAvaBroadcastEditor> BroadcastEditor;
 	
-	TWeakObjectPtr<UAvalancheBroadcast> AvalancheBroadcast;
+	TWeakObjectPtr<UAvaBroadcast> BroadcastWeak;
 	
-	TWeakPtr<FAvaOutputTileItem> SelectedOutputTileWeak;	
+	TWeakPtr<FAvaBroadcastOutputTileItem> SelectedOutputTileWeak;	
 };
diff --git a/Engine/Plugins/Experimental/Avalanche/Source/AvalancheMediaEditor/Public/IAvaMediaEditorModule.h b/Engine/Plugins/Experimental/Avalanche/Source/AvalancheMediaEditor/Public/IAvaMediaEditorModule.h
--- a/Engine/Plugins/Experimental/Avalanche/Source/AvalancheMediaEditor/Public/IAvaMediaEditorModule.h
+++ b/Engine/Plugins/Experimental/Avalanche/Source/AvalancheMediaEditor/Public/IAvaMediaEditorModule.h
@@ -9,11 +9,11 @@
 
 DECLARE_LOG_CATEGORY_EXTERN(LogAvaMediaEditor, Log, All);
 
-class FAvaPlaylistServer;
-class IAvaPlaylistFilterSuggestionFactory;
-struct FAvalanchePage;
-struct FAvaPlaylistTextFilterArgs;
-enum class EAvaPlaylistSearchListType : uint8;
+class FAvaRundownServer;
+class IAvaRundownFilterSuggestionFactory;
+enum class EAvaRundownSearchListType : uint8;
+struct FAvaRundownPage;
+struct FAvaRundownTextFilterArgs;
 
 class IAvaMediaEditorModule : public IModuleInterface
 {
@@ -31,9 +31,9 @@ public:
 	}
 
 	/** Returns the tool menu name used for Page Context Menu */
-	static FName GetPlaylistPageMenuName()
+	static FName GetRundownPageMenuName()
 	{
-		return TEXT("AvaPlaylistPageContextMenu");
+		return TEXT("AvaRundownPageContextMenu");
 	}
 
 	virtual FSlateIcon GetToolbarBroadcastButtonIcon() const = 0;
@@ -45,51 +45,51 @@ public:
 	virtual TSharedPtr<FExtensibilityManager> GetPlaybackToolBarExtensibilityManager() = 0;
 
 	/** Returns the toolbar extensibility manager for the Playback Editor */
-	virtual TSharedPtr<FExtensibilityManager> GetPlaylistToolBarExtensibilityManager() = 0;
+	virtual TSharedPtr<FExtensibilityManager> GetRundownToolBarExtensibilityManager() = 0;
 
 	/**
-	 * Returns the context menu extensibility manager for the Playlist Editor's Template Pages
-	 * @remark prefer extending with the UToolMenu named after IAvaMediaEditorModule::GetPlaylistPageMenuName, and using UAvaPlaylistPageContext to retrieve context information
+	 * Returns the context menu extensibility manager for the Rundown Editor's Template Pages
+	 * @remark prefer extending with the UToolMenu named after IAvaMediaEditorModule::GetRundownPageMenuName, and using UAvaRundownPageContext to retrieve context information
 	 */
-	virtual TSharedPtr<FExtensibilityManager> GetPlaylistMenuExtensibilityManager() = 0;
+	virtual TSharedPtr<FExtensibilityManager> GetRundownMenuExtensibilityManager() = 0;
 
-	DECLARE_MULTICAST_DELEGATE(FOnPlaylistServerStarted);
-	virtual FOnPlaylistServerStarted& GetOnPlaylistServerStarted() = 0;
+	DECLARE_MULTICAST_DELEGATE(FOnRundownServerStarted);
+	virtual FOnRundownServerStarted& GetOnRundownServerStarted() = 0;
 
-	DECLARE_MULTICAST_DELEGATE(FOnPlaylistServerStopped);
-	virtual FOnPlaylistServerStopped& GetOnPlaylistServerStopped() = 0;
+	DECLARE_MULTICAST_DELEGATE(FOnRundownServerStopped);
+	virtual FOnRundownServerStopped& GetOnRundownServerStopped() = 0;
 
-	virtual TSharedPtr<FAvaPlaylistServer> GetPlaylistServer() const = 0;
+	virtual TSharedPtr<FAvaRundownServer> GetRundownServer() const = 0;
 
 	/**
- 	* Check if current playlist filter expression factory support the comparison operation
- 	* @param InFilterKey Filter key to get the playlist filter expression factory needed
+ 	* Check if current rundown filter expression factory support the comparison operation
+ 	* @param InFilterKey Filter key to get the rundown filter expression factory needed
  	* @param InOperation Operation to check if supported
- 	* @param InPlaylistSearchListType Type of the Search List either Template or Instanced
+ 	* @param InRundownSearchListType Type of the Search List either Template or Instanced
  	* @return True if operation is supported, False otherwise
  	*/
-	virtual bool CanFilterSupportComparisonOperation(const FName& InFilterKey, ETextFilterComparisonOperation InOperation, EAvaPlaylistSearchListType InPlaylistSearchListType) const = 0;
+	virtual bool CanFilterSupportComparisonOperation(const FName& InFilterKey, ETextFilterComparisonOperation InOperation, EAvaRundownSearchListType InRundownSearchListType) const = 0;
 
 	/**
  	* Evaluate the expression and return the result
  	* @param InFilterKey Filter Key to get the Factory
  	* @param InItem Item that is currently checked
- 	* @param InArgs Args to evaluate the expression see FAvaPlaylistTextFilterArgs for more information
+ 	* @param InArgs Args to evaluate the expression see FAvaRundownTextFilterArgs for more information
  	* @return True if the expression evaluated to True, False otherwise
  	*/
-	virtual bool FilterExpression(const FName& InFilterKey, const FAvalanchePage& InItem, const FAvaPlaylistTextFilterArgs& InArgs) const = 0;
+	virtual bool FilterExpression(const FName& InFilterKey, const FAvaRundownPage& InItem, const FAvaRundownTextFilterArgs& InArgs) const = 0;
 
 	/**
  	* Get all simple suggestions with the given type (Template/Instanced/All)
- 	* @param InSuggestionType Type of suggestion to get, see EAvaPlaylistSearchListType for more information
+ 	* @param InSuggestionType Type of suggestion to get, see EAvaRundownSearchListType for more information
  	* @return An Array containing all suggestions of the given type
  	*/
-	virtual TArray<TSharedPtr<IAvaPlaylistFilterSuggestionFactory>> GetSimpleSuggestions(EAvaPlaylistSearchListType InSuggestionType) const = 0;
+	virtual TArray<TSharedPtr<IAvaRundownFilterSuggestionFactory>> GetSimpleSuggestions(EAvaRundownSearchListType InSuggestionType) const = 0;
 
 	/**
  	* Get all complex suggestions with the given type (Template/Instanced/All)
- 	* @param InSuggestionType Type of suggestion to get, see EAvaPlaylistSearchListType for more information
+ 	* @param InSuggestionType Type of suggestion to get, see EAvaRundownSearchListType for more information
  	* @return An Array containing all suggestions of the given type
  	*/
-	virtual TArray<TSharedPtr<IAvaPlaylistFilterSuggestionFactory>> GetComplexSuggestions(EAvaPlaylistSearchListType InSuggestionType) const = 0;
+	virtual TArray<TSharedPtr<IAvaRundownFilterSuggestionFactory>> GetComplexSuggestions(EAvaRundownSearchListType InSuggestionType) const = 0;
 };
diff --git a/Engine/Plugins/Experimental/Avalanche/Source/AvalancheMediaEditor/Public/Playback/AvaPlaybackCommands.h b/Engine/Plugins/Experimental/Avalanche/Source/AvalancheMediaEditor/Public/Playback/AvaPlaybackCommands.h
--- a/Engine/Plugins/Experimental/Avalanche/Source/AvalancheMediaEditor/Public/Playback/AvaPlaybackCommands.h
+++ b/Engine/Plugins/Experimental/Avalanche/Source/AvalancheMediaEditor/Public/Playback/AvaPlaybackCommands.h
@@ -9,7 +9,7 @@ class FAvaPlaybackCommands : public TCommands<FAvaPlaybackCommands>
 {
 public:
 	
-	FAvaPlaybackCommands() : TCommands<FAvaPlaybackCommands>(TEXT("AvalanchePlayback")
+	FAvaPlaybackCommands() : TCommands<FAvaPlaybackCommands>(TEXT("AvaPlayback")
 		, NSLOCTEXT("AvaPlaybackCommands", "AvaPlaybackCommands", "Motion Design Playback")
 		, NAME_None
 		, FAppStyle::GetAppStyleSetName())
diff --git a/Engine/Plugins/Experimental/Avalanche/Source/AvalancheMediaEditor/Public/Playback/AvaPlaybackGraphFactory.h b/Engine/Plugins/Experimental/Avalanche/Source/AvalancheMediaEditor/Public/Playback/AvaPlaybackGraphFactory.h
new file mode 100644
--- /dev/null
+++ b/Engine/Plugins/Experimental/Avalanche/Source/AvalancheMediaEditor/Public/Playback/AvaPlaybackGraphFactory.h
@@ -0,0 +1,31 @@
+// Copyright Epic Games, Inc. All Rights Reserved.
+
+#pragma once
+
+#include "CoreMinimal.h"
+#include "Factories/Factory.h"
+#include "AvaPlaybackGraphFactory.generated.h"
+
+UCLASS()
+class AVALANCHEMEDIAEDITOR_API UAvaPlaybackGraphFactory : public UFactory
+{
+	GENERATED_BODY()
+
+public:
+
+	UAvaPlaybackGraphFactory();
+	virtual ~UAvaPlaybackGraphFactory() override;
+
+protected:
+	
+	//~ Begin UFactory Interface
+	virtual uint32 GetMenuCategories() const override;
+	virtual bool ConfigureProperties() override;
+	virtual bool ShouldShowInNewMenu() const override;
+	virtual UObject* FactoryCreateNew(UClass* Class, UObject* InParent, FName Name, EObjectFlags Flags, UObject* Context, FFeedbackContext* Warn, FName CallingContext) override;
+	virtual UObject* FactoryCreateNew(UClass* Class, UObject* InParent, FName Name, EObjectFlags Flags, UObject* Context, FFeedbackContext* Warn) override;
+	virtual bool DoesSupportClass(UClass* Class) override;
+	virtual UClass* ResolveSupportedClass() override;
+	virtual FString GetDefaultNewAssetName() const override;
+	//~ Begin UFactory Interface
+};
diff --git a/Engine/Plugins/Experimental/Avalanche/Source/AvalancheMediaEditor/Public/Playback/Graph/Nodes/AvaPlaybackEditorGraphNode_BlueprintPlayer.h b/Engine/Plugins/Experimental/Avalanche/Source/AvalancheMediaEditor/Public/Playback/Graph/Nodes/AvaPlaybackEditorGraphNode_BlueprintPlayer.h
new file mode 100644
--- /dev/null
+++ b/Engine/Plugins/Experimental/Avalanche/Source/AvalancheMediaEditor/Public/Playback/Graph/Nodes/AvaPlaybackEditorGraphNode_BlueprintPlayer.h
@@ -0,0 +1,20 @@
+// Copyright Epic Games, Inc. All Rights Reserved.
+
+#pragma once
+
+#include "AvaPlaybackEditorGraphNode.h"
+#include "AvaPlaybackEditorGraphNode_BlueprintPlayer.generated.h"
+
+UCLASS()
+class AVALANCHEMEDIAEDITOR_API UAvaPlaybackEditorGraphNode_BlueprintPlayer : public UAvaPlaybackEditorGraphNode
+{
+	GENERATED_BODY()
+	
+	virtual TSubclassOf<UAvaPlaybackNode> GetPlaybackNodeClass() const override;
+	virtual FName GetInputPinCategory(int32 InputPinIndex) const override;
+	
+	//UEdGraphNode interface
+	virtual FLinearColor GetNodeTitleColor() const override;
+	virtual TSharedPtr<SGraphNode> CreateVisualWidget() override;
+	//~UEdGraphNode interface
+};
diff --git a/Engine/Plugins/Experimental/Avalanche/Source/AvalancheMediaEditor/Public/Rundown/AvaRundownDefines.h b/Engine/Plugins/Experimental/Avalanche/Source/AvalancheMediaEditor/Public/Rundown/AvaRundownDefines.h
new file mode 100644
--- /dev/null
+++ b/Engine/Plugins/Experimental/Avalanche/Source/AvalancheMediaEditor/Public/Rundown/AvaRundownDefines.h
@@ -0,0 +1,50 @@
+// Copyright Epic Games, Inc. All Rights Reserved.
+
+#pragma once
+
+#include "Internationalization/Text.h"
+#include "Templates/SharedPointer.h"
+
+class IAvaRundownInstancedPageView;
+class IAvaRundownPageView;
+class IAvaRundownTemplatePageView;
+struct FAvaRundownPage;
+
+typedef TWeakPtr<IAvaRundownPageView> FAvaRundownPageViewWeak;
+typedef TSharedPtr<IAvaRundownPageView> FAvaRundownPageViewPtr;
+typedef TSharedRef<IAvaRundownPageView> FAvaRundownPageViewRef;
+typedef TWeakPtr<IAvaRundownTemplatePageView> FAvaRundownTemplatePageViewWeak;
+typedef TSharedPtr<IAvaRundownTemplatePageView> FAvaRundownTemplatePageViewPtr;
+typedef TSharedRef<IAvaRundownTemplatePageView> FAvaRundownTemplatePageViewRef;
+typedef TWeakPtr<IAvaRundownInstancedPageView> FAvaRundownInstancedPageViewWeak;
+typedef TSharedPtr<IAvaRundownInstancedPageView> FAvaRundownInstancedPageViewPtr;
+typedef TSharedRef<IAvaRundownInstancedPageView> FAvaRundownInstancedPageViewRef;
+
+namespace UE::AvaRundown
+{
+	enum {InvalidPageId	= -1};
+	
+	struct AVALANCHEMEDIAEDITOR_API FEditorMetrics
+	{
+		static constexpr float ColumnLeftOffset = 5.f;
+		static const FNumberFormattingOptions PageIdFormattingOptions;
+	};
+	
+	enum class EPageEvent : uint8
+	{
+		/** Called after the Selection has Changed */ 
+		SelectionChanged,
+
+		/** Called to sync the Selection of current Selected Items in the Caller */
+		SelectionRequest,
+
+		/** Called to Request a Rename of the Selected Items */
+		RenameRequest,
+
+		/** Called to Request Renumbering Page Id of Selected Items */
+		RenumberRequest,
+
+		/** Called to Request Reimporting Assets of Selected Items */
+		ReimportRequest,
+	};
+}
diff --git a/Engine/Plugins/Experimental/Avalanche/Source/AvalancheMediaEditor/Public/Rundown/AvaRundownFactory.h b/Engine/Plugins/Experimental/Avalanche/Source/AvalancheMediaEditor/Public/Rundown/AvaRundownFactory.h
new file mode 100644
--- /dev/null
+++ b/Engine/Plugins/Experimental/Avalanche/Source/AvalancheMediaEditor/Public/Rundown/AvaRundownFactory.h
@@ -0,0 +1,31 @@
+// Copyright Epic Games, Inc. All Rights Reserved.
+
+#pragma once
+
+#include "CoreMinimal.h"
+#include "Factories/Factory.h"
+#include "AvaRundownFactory.generated.h"
+
+UCLASS()
+class AVALANCHEMEDIAEDITOR_API UAvaRundownFactory : public UFactory
+{
+	GENERATED_BODY()
+
+public:
+
+	UAvaRundownFactory();
+	virtual ~UAvaRundownFactory() override;
+
+protected:
+	
+	//~ Begin UFactory Interface
+	virtual uint32 GetMenuCategories() const override;
+	virtual bool ConfigureProperties() override;
+	virtual bool ShouldShowInNewMenu() const override;
+	virtual UObject* FactoryCreateNew(UClass* Class, UObject* InParent, FName Name, EObjectFlags Flags, UObject* Context, FFeedbackContext* Warn, FName CallingContext) override;
+	virtual UObject* FactoryCreateNew(UClass* Class, UObject* InParent, FName Name, EObjectFlags Flags, UObject* Context, FFeedbackContext* Warn) override;
+	virtual bool DoesSupportClass(UClass* Class) override;
+	virtual UClass* ResolveSupportedClass() override;
+	virtual FString GetDefaultNewAssetName() const override;
+	//~ Begin UFactory Interface
+};
diff --git a/Engine/Plugins/Experimental/Avalanche/Source/AvalancheMediaEditor/Public/Rundown/AvaRundownMessages.h b/Engine/Plugins/Experimental/Avalanche/Source/AvalancheMediaEditor/Public/Rundown/AvaRundownMessages.h
new file mode 100644
--- /dev/null
+++ b/Engine/Plugins/Experimental/Avalanche/Source/AvalancheMediaEditor/Public/Rundown/AvaRundownMessages.h
@@ -0,0 +1,688 @@
+// Copyright Epic Games, Inc. All Rights Reserved.
+
+#pragma once
+
+#include "Rundown/AvaRundownPage.h"
+#include "AvaRundownMessages.generated.h"
+
+namespace EAvaRundownApiVersion
+{
+	/**
+	 * Defines the protocol version of the Rundown Server API.
+	 *
+	 * API versioning is used to provide legacy support either on
+	 * the client side or server side for non compatible changes.
+	 * Clients can request a version of the API that they where implemented against,
+	 * if the server can still honor the request it will accept.
+	 */
+	enum Type
+	{
+		Initial = 1,
+
+		// -----<new versions can be added before this line>-------------------------------------------------
+		// - this needs to be the last line (see note below)
+		VersionPlusOne,
+		LatestVersion = VersionPlusOne - 1
+	};
+}
+
+
+USTRUCT()
+struct AVALANCHEMEDIAEDITOR_API FAvaRundownMsgBase
+{
+	GENERATED_BODY()
+public:
+	UPROPERTY()
+	int32 RequestId = -1;
+};
+
+USTRUCT()
+struct AVALANCHEMEDIAEDITOR_API FAvaRundownServerMsg : public FAvaRundownMsgBase
+{
+	GENERATED_BODY()
+public:
+	UPROPERTY()
+	FString Verbosity;
+
+	UPROPERTY()
+	FString Text;
+};
+
+/** Request published by client to discover servers. */
+USTRUCT()
+struct AVALANCHEMEDIAEDITOR_API FAvaRundownPing : public FAvaRundownMsgBase
+{
+	GENERATED_BODY()
+public:
+	/** True if the request originates from an automatic timer. False if requests originates from user interaction. */
+	UPROPERTY()
+	bool bAuto = true;
+
+	/**
+	 * API Version the client has been implemented against.
+	 * If none (-1) the server will consider the initial version is requested.
+	 */
+	UPROPERTY()
+	int32 RequestedApiVersion = -1;
+};
+
+/** Response sent by server to client to be discovered. */
+USTRUCT()
+struct AVALANCHEMEDIAEDITOR_API FAvaRundownPong : public FAvaRundownMsgBase
+{
+	GENERATED_BODY()
+
+public:
+	/** True if it is a reply to an auto ping. Mirrors the bAuto flag from Ping message. */
+	UPROPERTY()
+	bool bAuto = true;
+
+	/**
+	 * API Version the server will communicate with for this client.
+	 * The server may honor the requested version if possible.
+	 * Versions newer than server implementation will obviously not be honored either.
+	 * Clients should expect an older server to reply with an older version.
+	 */
+	UPROPERTY()
+	int32 ApiVersion = -1;
+
+	/** Minimum API Version the server implements. */
+	UPROPERTY()
+	int32 MinimumApiVersion = -1;
+
+	/** Latest API Version the server support. */
+	UPROPERTY()
+	int32 LatestApiVersion = -1;
+
+	UPROPERTY()
+	FString HostName;
+};
+
+/**
+ *	Request list of rundown that can be opened on the current server.
+ */
+USTRUCT()
+struct AVALANCHEMEDIAEDITOR_API FAvaRundownGetRundowns : public FAvaRundownMsgBase
+{
+	GENERATED_BODY()
+};
+
+/**
+ *	List of all rundowns.
+ *	Expected Response from FAvaRundownGetRundowns.
+ */
+USTRUCT()
+struct AVALANCHEMEDIAEDITOR_API FAvaRundownRundowns : public FAvaRundownMsgBase
+{
+	GENERATED_BODY()
+
+public:
+	UPROPERTY()
+	TArray<FString> Rundowns;
+};
+
+/**
+ *	Request that the given rundown be opened.
+ *	Only one rundown can be opened at a time. If another rundown
+ *	is opened, it will be closed and all currently playing pages stopped.
+ *	If the path is empty, nothing will be done and the server will reply with
+ *	a FAvaRundownLoadedRundown message indicating which rundown is currently loaded.
+ */
+USTRUCT()
+struct AVALANCHEMEDIAEDITOR_API FAvaRundownLoadRundown : public FAvaRundownMsgBase
+{
+	GENERATED_BODY()
+
+public:
+	UPROPERTY()
+	FString Rundown;
+};
+
+/**
+ * Request the list of pages from the given rundown.
+ */
+USTRUCT()
+struct AVALANCHEMEDIAEDITOR_API FAvaRundownGetPages : public FAvaRundownMsgBase
+{
+	GENERATED_BODY()
+public:
+	UPROPERTY()
+	FString Rundown;
+};
+
+USTRUCT()
+struct AVALANCHEMEDIAEDITOR_API FAvaRundownCreatePage : public FAvaRundownMsgBase
+{
+	GENERATED_BODY()
+
+	UPROPERTY()
+	FString Rundown;
+
+	UPROPERTY()
+	int32 TemplateId = FAvaRundownPage::InvalidPageId;
+};
+
+USTRUCT()
+struct AVALANCHEMEDIAEDITOR_API FAvaRundownDeletePage : public FAvaRundownMsgBase
+{
+	GENERATED_BODY()
+
+	UPROPERTY()
+	FString Rundown;
+
+	UPROPERTY()
+	int32 PageId = FAvaRundownPage::InvalidPageId;
+};
+
+USTRUCT()
+struct AVALANCHEMEDIAEDITOR_API FAvaRundownCreateTemplate : public FAvaRundownMsgBase
+{
+	GENERATED_BODY()
+
+	UPROPERTY()
+	FString Rundown;
+};
+
+USTRUCT()
+struct AVALANCHEMEDIAEDITOR_API FAvaRundownDeleteTemplate : public FAvaRundownMsgBase
+{
+	GENERATED_BODY()
+
+	UPROPERTY()
+	FString Rundown;
+
+	UPROPERTY()
+	int32 PageId = FAvaRundownPage::InvalidPageId;
+};
+
+USTRUCT()
+struct AVALANCHEMEDIAEDITOR_API FAvaRundownChangeTemplateBP : public FAvaRundownMsgBase
+{
+	GENERATED_BODY()
+public:
+	UPROPERTY()
+	FString Rundown;
+
+	UPROPERTY()
+	int32 TemplateId = FAvaRundownPage::InvalidPageId;
+
+	UPROPERTY()
+	FString AssetPath;
+};
+
+USTRUCT()
+struct AVALANCHEMEDIAEDITOR_API FAvaRundownPageInfo
+{
+	GENERATED_BODY()
+public:
+	UPROPERTY()
+	int32 PageId = FAvaRundownPage::InvalidPageId;
+
+	UPROPERTY()
+	FString PageName;
+
+	UPROPERTY()
+	bool IsTemplate = false;
+
+	UPROPERTY()
+	int32 TemplateId = FAvaRundownPage::InvalidPageId;
+
+	UPROPERTY()
+	TArray<int32> CombinedTemplateIds;
+
+	UPROPERTY()
+	FSoftObjectPath AssetPath;
+
+	UPROPERTY()
+	TArray<FAvaRundownChannelPageStatus> Statuses;
+
+	UPROPERTY()
+	FString TransitionLayerName;
+
+	UPROPERTY()
+	FString OutputChannel;
+
+	UPROPERTY()
+	bool bIsEnabled = false;
+
+	UPROPERTY()
+	bool bIsPlaying = false;
+};
+
+/*
+ * List of pages from the current rundown.
+ */
+USTRUCT()
+struct AVALANCHEMEDIAEDITOR_API FAvaRundownPages : public FAvaRundownMsgBase
+{
+	GENERATED_BODY()
+public:
+	UPROPERTY()
+	TArray<FAvaRundownPageInfo> Pages;
+};
+
+/**
+ * Request the page details from the given rundown.
+ */
+USTRUCT()
+struct AVALANCHEMEDIAEDITOR_API FAvaRundownGetPageDetails : public FAvaRundownMsgBase
+{
+	GENERATED_BODY()
+
+public:
+	UPROPERTY()
+	FString Rundown;
+
+	UPROPERTY()
+	int32 PageId = FAvaRundownPage::InvalidPageId;
+
+	/** This will request that a managed asset instance gets loaded to be
+	 * accessible through WebRC. */
+	UPROPERTY()
+	bool bLoadRemoteControlPreset = false;
+};
+
+/**
+ *	Server response to FAvaRundownGetPageDetails request.
+ */
+USTRUCT()
+struct AVALANCHEMEDIAEDITOR_API FAvaRundownPageDetails : public FAvaRundownMsgBase
+{
+	GENERATED_BODY()
+
+public:
+	UPROPERTY()
+	FString Rundown;
+
+	UPROPERTY()
+	FAvaRundownPageInfo PageInfo;
+
+	UPROPERTY()
+	FAvaPlayableRemoteControlValues RemoteControlValues;
+
+	/** Name of the remote control preset to resolve through WebRC API. */
+	UPROPERTY()
+	FString RemoteControlPresetName;
+
+	UPROPERTY()
+	FString RemoteControlPresetId;
+};
+
+USTRUCT()
+struct AVALANCHEMEDIAEDITOR_API FAvaRundownPagesStatuses : public FAvaRundownMsgBase
+{
+	GENERATED_BODY()
+
+public:
+	UPROPERTY()
+	FString Rundown;
+
+	UPROPERTY()
+	FAvaRundownPageInfo PageInfo;
+};
+
+USTRUCT()
+struct AVALANCHEMEDIAEDITOR_API FAvaRundownPageListChanged : public FAvaRundownMsgBase
+{
+	GENERATED_BODY()
+
+	UPROPERTY()
+	FString Rundown;
+
+	/** See EAvaPageListChange flags. */
+	UPROPERTY()
+	uint8 ChangeType = 0;
+
+	UPROPERTY();
+	TArray<int32> AffectedPages;
+};
+
+USTRUCT()
+struct AVALANCHEMEDIAEDITOR_API FAvaRundownPageBlueprintChanged : public FAvaRundownMsgBase
+{
+	GENERATED_BODY()
+
+	UPROPERTY()
+	FString Rundown;
+
+	UPROPERTY()
+	int32 PageId = FAvaRundownPage::InvalidPageId;
+
+	UPROPERTY()
+	FString BlueprintPath;
+};
+
+USTRUCT()
+struct AVALANCHEMEDIAEDITOR_API FAvaRundownPageChannelChanged : public FAvaRundownMsgBase
+{
+	GENERATED_BODY()
+
+	UPROPERTY()
+	FString Rundown;
+
+	UPROPERTY()
+	int32 PageId = FAvaRundownPage::InvalidPageId;
+
+	UPROPERTY()
+	FString ChannelName;
+};
+
+USTRUCT()
+struct AVALANCHEMEDIAEDITOR_API FAvaRundownPageAnimSettingsChanged : public FAvaRundownMsgBase
+{
+	GENERATED_BODY()
+
+	UPROPERTY()
+	FString Rundown;
+
+	UPROPERTY()
+	int32 PageId = FAvaRundownPage::InvalidPageId;
+};
+
+USTRUCT()
+struct AVALANCHEMEDIAEDITOR_API FAvaRundownPageChangeChannel : public FAvaRundownMsgBase
+{
+	GENERATED_BODY()
+
+	UPROPERTY()
+	FString Rundown;
+
+	UPROPERTY()
+	int32 PageId = FAvaRundownPage::InvalidPageId;
+
+	UPROPERTY()
+	FString ChannelName;
+};
+
+/** This is a request to save the managed RCP back to the corresponding page. */
+USTRUCT()
+struct AVALANCHEMEDIAEDITOR_API FAvaRundownUpdatePageFromRCP : public FAvaRundownMsgBase
+{
+	GENERATED_BODY()
+
+public:
+	/** Unregister the Remote Control Preset from the WebRC. */
+	UPROPERTY()
+	bool bUnregister = false;
+};
+
+/** Supported Page actions for playback. */
+UENUM()
+enum class EAvaRundownPageActions
+{
+	None,
+	Load,
+	Unload,
+	Play,
+	PlayNext,
+	Stop,
+	ForceStop,
+	Continue,
+	UpdateValues,
+	TakeToProgram
+};
+
+USTRUCT()
+struct AVALANCHEMEDIAEDITOR_API FAvaRundownPageAction : public FAvaRundownMsgBase
+{
+	GENERATED_BODY()
+public:
+	UPROPERTY()
+	int32 PageId = FAvaRundownPage::InvalidPageId;
+
+	UPROPERTY()
+	EAvaRundownPageActions Action = EAvaRundownPageActions::None;
+};
+
+USTRUCT()
+struct AVALANCHEMEDIAEDITOR_API FAvaRundownPagePreviewAction : public FAvaRundownPageAction
+{
+	GENERATED_BODY()
+public:
+	/** Specify which preview channel to use. If left empty, the rundown's default preview channel is used. */
+	UPROPERTY()
+	FString PreviewChannelName;
+};
+
+/**
+ * Command to execute an action on multiple pages at the same time.
+ * This is necessary for pages to be part of the same transition.
+ */
+USTRUCT()
+struct AVALANCHEMEDIAEDITOR_API FAvaRundownPageActions : public FAvaRundownMsgBase
+{
+	GENERATED_BODY()
+public:
+	UPROPERTY()
+	TArray<int32> PageIds;
+
+	UPROPERTY()
+	EAvaRundownPageActions Action = EAvaRundownPageActions::None;
+};
+
+USTRUCT()
+struct AVALANCHEMEDIAEDITOR_API FAvaRundownPagePreviewActions : public FAvaRundownPageActions
+{
+	GENERATED_BODY()
+public:
+	/** Specify which preview channel to use. If left empty, the rundown's default preview channel is used. */
+	UPROPERTY()
+	FString PreviewChannelName;
+};
+
+UENUM()
+enum class EAvaRundownPageEvents
+{
+	None,
+	AnimStarted,
+	AnimPaused,
+	AnimFinished
+};
+
+USTRUCT()
+struct AVALANCHEMEDIAEDITOR_API FAvaRundownPageEvent : public FAvaRundownMsgBase
+{
+	GENERATED_BODY()
+public:
+	UPROPERTY()
+	int32 PageId = FAvaRundownPage::InvalidPageId;
+
+	UPROPERTY()
+	EAvaRundownPageEvents Event = EAvaRundownPageEvents::None;
+};
+
+USTRUCT()
+struct AVALANCHEMEDIAEDITOR_API FAvaRundownOutputDeviceItem
+{
+	GENERATED_BODY()
+public:
+	UPROPERTY()
+	FString Name;
+
+	/**
+	 * Raw Json string representing a serialized UMediaOutput.
+	 */
+	UPROPERTY()
+	FString Data;
+};
+
+USTRUCT()
+struct AVALANCHEMEDIAEDITOR_API FAvaRundownOutputClassItem
+{
+	GENERATED_BODY()
+public:
+	UPROPERTY()
+	FString Name;
+
+	UPROPERTY()
+	TArray<FAvaRundownOutputDeviceItem> Devices;
+};
+
+USTRUCT()
+struct AVALANCHEMEDIAEDITOR_API FAvaRundownDevicesList : public FAvaRundownMsgBase
+{
+	GENERATED_BODY()
+public:
+	UPROPERTY()
+	TArray<FAvaRundownOutputClassItem> DeviceClasses;
+};
+
+USTRUCT()
+struct AVALANCHEMEDIAEDITOR_API FAvaRundownGetChannel : public FAvaRundownMsgBase
+{
+	GENERATED_BODY()
+public:
+	UPROPERTY()
+	FString ChannelName;
+};
+
+USTRUCT()
+struct AVALANCHEMEDIAEDITOR_API FAvaRundownGetChannels : public FAvaRundownMsgBase
+{
+	GENERATED_BODY()
+};
+
+USTRUCT()
+struct AVALANCHEMEDIAEDITOR_API FAvaRundownChannel
+{
+	GENERATED_BODY()
+public:
+	UPROPERTY()
+	FString Name;
+
+	UPROPERTY()
+	TArray<FAvaRundownOutputDeviceItem> Devices;
+};
+
+USTRUCT()
+struct AVALANCHEMEDIAEDITOR_API FAvaRundownChannelListChanged : public FAvaRundownMsgBase
+{
+	GENERATED_BODY()
+
+	UPROPERTY()
+	TArray<FAvaRundownChannel> Channels;
+};
+
+USTRUCT()
+struct AVALANCHEMEDIAEDITOR_API FAvaRundownChannelResponse : public FAvaRundownMsgBase
+{
+	GENERATED_BODY()
+public:
+	UPROPERTY()
+	FAvaRundownChannel Channel;
+};
+
+USTRUCT()
+struct AVALANCHEMEDIAEDITOR_API FAvaRundownChannels : public FAvaRundownMsgBase
+{
+	GENERATED_BODY()
+public:
+	UPROPERTY()
+	TArray<FAvaRundownChannel> Channels;
+};
+
+USTRUCT()
+struct AVALANCHEMEDIAEDITOR_API FAvaRundownAssetsChanged : public FAvaRundownMsgBase
+{
+	GENERATED_BODY()
+
+	UPROPERTY()
+	FString AssetName;
+};
+
+// Channel actions
+UENUM()
+enum class EAvaRundownChannelActions
+{
+	None,
+	Start,
+	Stop
+};
+
+USTRUCT()
+struct AVALANCHEMEDIAEDITOR_API FAvaRundownChannelAction : public FAvaRundownMsgBase
+{
+	GENERATED_BODY()
+public:
+	UPROPERTY()
+	FString ChannelName;
+
+	UPROPERTY()
+	EAvaRundownChannelActions Action = EAvaRundownChannelActions::None;
+};
+
+USTRUCT()
+struct AVALANCHEMEDIAEDITOR_API FAvaRundownGetDevices : public FAvaRundownMsgBase
+{
+	GENERATED_BODY()
+};
+
+USTRUCT()
+struct AVALANCHEMEDIAEDITOR_API FAvaRundownAddChannelDevice : public FAvaRundownMsgBase
+{
+	GENERATED_BODY()
+public:
+	UPROPERTY()
+	FString ChannelName;
+
+	UPROPERTY()
+	FString MediaOutputName;
+};
+
+USTRUCT()
+struct AVALANCHEMEDIAEDITOR_API FAvaRundownEditChannelDevice : public FAvaRundownMsgBase
+{
+	GENERATED_BODY()
+public:
+	UPROPERTY()
+	FString ChannelName;
+
+	UPROPERTY()
+	FString MediaOutputName;
+
+	UPROPERTY()
+	FString Data;
+};
+
+USTRUCT()
+struct AVALANCHEMEDIAEDITOR_API FAvaRundownRemoveChannelDevice : public FAvaRundownMsgBase
+{
+	GENERATED_BODY()
+public:
+	UPROPERTY()
+	FString ChannelName;
+
+	UPROPERTY()
+	FString MediaOutputName;
+};
+
+/* No difference from FAvaRundownOutputDeviceItem except this is meant to return a single device response. */
+USTRUCT()
+struct AVALANCHEMEDIAEDITOR_API FAvaRundownOutputDeviceItemResponse : public FAvaRundownMsgBase
+{
+	GENERATED_BODY()
+public:
+	UPROPERTY()
+	FString Name;
+
+	UPROPERTY()
+	FString Data;
+};
+
+USTRUCT()
+struct AVALANCHEMEDIAEDITOR_API FAvaRundownGetChannelImage : public FAvaRundownMsgBase
+{
+	GENERATED_BODY()
+public:
+	UPROPERTY()
+	FString ChannelName;
+};
+
+USTRUCT()
+struct AVALANCHEMEDIAEDITOR_API FAvaRundownChannelImage : public FAvaRundownMsgBase
+{
+	GENERATED_BODY()
+public:
+	UPROPERTY()
+	TArray<uint8> ImageData;
+};
diff --git a/Engine/Plugins/Experimental/StormSyncAvalancheBridge/Source/StormSyncAvaBridge/Private/ModularFeature/StormSyncAvaSyncProvider.cpp b/Engine/Plugins/Experimental/StormSyncAvalancheBridge/Source/StormSyncAvaBridge/Private/ModularFeature/StormSyncAvaSyncProvider.cpp
--- a/Engine/Plugins/Experimental/StormSyncAvalancheBridge/Source/StormSyncAvaBridge/Private/ModularFeature/StormSyncAvaSyncProvider.cpp
+++ b/Engine/Plugins/Experimental/StormSyncAvalancheBridge/Source/StormSyncAvaBridge/Private/ModularFeature/StormSyncAvaSyncProvider.cpp
@@ -1,11 +1,11 @@
 // Copyright Epic Games, Inc. All Rights Reserved.
 
 #include "StormSyncAvaSyncProvider.h"
-
 #include "IAvaMediaModule.h"
 #include "IStormSyncTransportClientModule.h"
-#include "Playback/AvaMediaPlaybackServer.h"
-#include "Playback/IAvaMediaPlaybackClient.h"
+#include "ModularFeature/AvaMediaSyncProviderFeatureTypes.h"
+#include "Playback/AvaPlaybackServer.h"
+#include "Playback/IAvaPlaybackClient.h"
 #include "StormSyncAvaBridgeCommon.h"
 #include "StormSyncAvaBridgeLog.h"
 #include "StormSyncCoreDelegates.h"
@@ -35,7 +35,7 @@ void FStormSyncAvaSyncProvider::SyncToAll(const TArray<FName>& InPackageNames)
 	IStormSyncTransportClientModule::Get().SynchronizePackages(PackageDescriptor, InPackageNames);
 }
 
-void FStormSyncAvaSyncProvider::PushToRemote(const FString& InRemoteName, const TArray<FName>& InPackageNames, const FOnAvaSyncResponse& DoneDelegate)
+void FStormSyncAvaSyncProvider::PushToRemote(const FString& InRemoteName, const TArray<FName>& InPackageNames, const FOnAvaMediaSyncResponse& DoneDelegate)
 {
 	STORM_SYNC_AVA_LOG(Display, TEXT("FStormSyncAvaSyncProvider::PushToRemote InRemoteName: %s, InPackageNames: %d"), *InRemoteName, InPackageNames.Num());
 	
@@ -61,8 +61,8 @@ void FStormSyncAvaSyncProvider::PushToRemote(const FString& InRemoteName, const
 
 		if (DoneDelegate.IsBound())
 		{
-			FAvaSyncResponse SyncResponse = ConvertSyncResponse(Response);
-			DoneDelegate.Execute(MakeShared<FAvaSyncResponse>(SyncResponse));
+			FAvaMediaSyncResponse SyncResponse = ConvertSyncResponse(Response);
+			DoneDelegate.Execute(MakeShared<FAvaMediaSyncResponse>(SyncResponse));
 		}
 	});
 
@@ -70,7 +70,7 @@ void FStormSyncAvaSyncProvider::PushToRemote(const FString& InRemoteName, const
 	IStormSyncTransportClientModule::Get().PushPackages(PackageDescriptor, InPackageNames, MessageAddress, Delegate);
 }
 
-void FStormSyncAvaSyncProvider::PullFromRemote(const FString& InRemoteName, const TArray<FName>& InPackageNames, const FOnAvaSyncResponse& DoneDelegate)
+void FStormSyncAvaSyncProvider::PullFromRemote(const FString& InRemoteName, const TArray<FName>& InPackageNames, const FOnAvaMediaSyncResponse& DoneDelegate)
 {
 	STORM_SYNC_AVA_LOG(Display, TEXT("FStormSyncAvaSyncProvider::PullFromRemote InRemoteName: %s, InPackageNames: %d"), *InRemoteName, InPackageNames.Num());
 
@@ -96,8 +96,8 @@ void FStormSyncAvaSyncProvider::PullFromRemote(const FString& InRemoteName, cons
 
 		if (DoneDelegate.IsBound())
 		{
-			FAvaSyncResponse SyncResponse = ConvertSyncResponse(Response);
-			DoneDelegate.Execute(MakeShared<FAvaSyncResponse>(SyncResponse));
+			FAvaMediaSyncResponse SyncResponse = ConvertSyncResponse(Response);
+			DoneDelegate.Execute(MakeShared<FAvaMediaSyncResponse>(SyncResponse));
 		}
 	});
 
@@ -105,7 +105,7 @@ void FStormSyncAvaSyncProvider::PullFromRemote(const FString& InRemoteName, cons
 	IStormSyncTransportClientModule::Get().PullPackages(PackageDescriptor, InPackageNames, MessageAddress, Delegate);
 }
 
-void FStormSyncAvaSyncProvider::CompareWithRemote(const FString& InRemoteName, const TArray<FName>& InPackageNames, const FOnAvaSyncCompareResponse& DoneDelegate)
+void FStormSyncAvaSyncProvider::CompareWithRemote(const FString& InRemoteName, const TArray<FName>& InPackageNames, const FOnAvaMediaSyncCompareResponse& DoneDelegate)
 {
 	STORM_SYNC_AVA_LOG(Display, TEXT("FStormSyncAvaSyncProvider::CompareWithRemote InRemoteName: %s, InPackageNames: %d"), *InRemoteName, InPackageNames.Num());
 
@@ -140,12 +140,12 @@ void FStormSyncAvaSyncProvider::CompareWithRemote(const FString& InRemoteName, c
 	CompareWith(MessageAddress, InPackageNames, DoneDelegate);
 }
 
-FOnAvaSyncPackageModified& FStormSyncAvaSyncProvider::GetOnAvaSyncPackageModified()
+FOnAvaMediaSyncPackageModified& FStormSyncAvaSyncProvider::GetOnAvaSyncPackageModified()
 {
 	return OnAvaSyncPackageModified;
 }
 
-void FStormSyncAvaSyncProvider::CompareWith(const FMessageAddress& InRemoteAddress, const TArray<FName>& InPackageNames, const FOnAvaSyncCompareResponse& DoneDelegate)
+void FStormSyncAvaSyncProvider::CompareWith(const FMessageAddress& InRemoteAddress, const TArray<FName>& InPackageNames, const FOnAvaMediaSyncCompareResponse& DoneDelegate)
 {
 	STORM_SYNC_AVA_LOG(Display, TEXT("FStormSyncAvaSyncProvider::CompareWith InPackageNames: %d (Remote Address: %s)"), InPackageNames.Num(), *InRemoteAddress.ToString());
 
@@ -166,11 +166,11 @@ void FStormSyncAvaSyncProvider::CompareWith(const FMessageAddress& InRemoteAddre
 
 		if (DoneDelegate.IsBound())
 		{
-			FAvaSyncCompareResponse CompareResponse;
-			CompareResponse.Status = EAvaSyncResponseResult::Success;
+			FAvaMediaSyncCompareResponse CompareResponse;
+			CompareResponse.Status = EAvaMediaSyncResponseResult::Success;
 			CompareResponse.ConnectionInfo = ConvertConnectionInfo(Response->ConnectionInfo);
 			CompareResponse.bNeedsSynchronization = Response->bNeedsSynchronization;
-			DoneDelegate.Execute(MakeShared<FAvaSyncCompareResponse>(CompareResponse));
+			DoneDelegate.Execute(MakeShared<FAvaMediaSyncCompareResponse>(CompareResponse));
 		}
 	}));
 }
@@ -227,7 +227,7 @@ bool FStormSyncAvaSyncProvider::GetAddressFromServerUserData(const FString& InSe
 		return false;
 	}
 
-	const IAvaMediaPlaybackClient& PlaybackClient = IAvaMediaModule::Get().GetMediaPlaybackClient();
+	const IAvaPlaybackClient& PlaybackClient = IAvaMediaModule::Get().GetMediaPlaybackClient();
 	const FString AddressId = PlaybackClient.GetServerUserData(InServerName, InUserDataKey);
 	if (AddressId.IsEmpty())
 	{
@@ -256,7 +256,7 @@ bool FStormSyncAvaSyncProvider::GetAddressFromClientUserData(const FString& InCl
 		return false;
 	}
 
-	const TSharedPtr<FAvaMediaPlaybackServer> PlaybackServer = IAvaMediaModule::Get().GetMediaPlaybackServer();
+	const TSharedPtr<FAvaPlaybackServer> PlaybackServer = IAvaMediaModule::Get().GetMediaPlaybackServer();
 	if (!PlaybackServer.IsValid())
 	{
 		if (OutErrorMessage)
@@ -293,27 +293,27 @@ bool FStormSyncAvaSyncProvider::IsPlaybackClient()
 	return IAvaMediaModule::Get().IsMediaPlaybackClientStarted();
 }
 
-TSharedPtr<FAvaSyncCompareResponse> FStormSyncAvaSyncProvider::CreateErrorResponse(const FText& InText)
+TSharedPtr<FAvaMediaSyncCompareResponse> FStormSyncAvaSyncProvider::CreateErrorResponse(const FText& InText)
 {
-	FAvaSyncCompareResponse Response;
-	Response.Status = EAvaSyncResponseResult::Error;
+	FAvaMediaSyncCompareResponse Response;
+	Response.Status = EAvaMediaSyncResponseResult::Error;
 	Response.ErrorText = InText;
-	return MakeShared<FAvaSyncCompareResponse>(Response);
+	return MakeShared<FAvaMediaSyncCompareResponse>(Response);
 }
 
-FAvaSyncResponse FStormSyncAvaSyncProvider::ConvertSyncResponse(const TSharedPtr<FStormSyncTransportSyncResponse>& InResponse)
+FAvaMediaSyncResponse FStormSyncAvaSyncProvider::ConvertSyncResponse(const TSharedPtr<FStormSyncTransportSyncResponse>& InResponse)
 {
 	check(InResponse.IsValid());
 	
-	FAvaSyncResponse AvaSyncResponse;
+	FAvaMediaSyncResponse AvaSyncResponse;
 	if (InResponse->Status == EStormSyncResponseResult::Success)
 	{
-		AvaSyncResponse.Status = EAvaSyncResponseResult::Success;
+		AvaSyncResponse.Status = EAvaMediaSyncResponseResult::Success;
 		AvaSyncResponse.StatusText = InResponse->StatusText;
 	}
 	else if (InResponse->Status == EStormSyncResponseResult::Error)
 	{
-		AvaSyncResponse.Status = EAvaSyncResponseResult::Error;
+		AvaSyncResponse.Status = EAvaMediaSyncResponseResult::Error;
 		AvaSyncResponse.ErrorText = InResponse->StatusText;
 	}
 
@@ -321,9 +321,9 @@ FAvaSyncResponse FStormSyncAvaSyncProvider::ConvertSyncResponse(const TSharedPtr
 	return AvaSyncResponse;
 }
 
-FAvaSyncConnectionInfo FStormSyncAvaSyncProvider::ConvertConnectionInfo(const FStormSyncConnectionInfo& InConnectionInfo)
+FAvaMediaSyncConnectionInfo FStormSyncAvaSyncProvider::ConvertConnectionInfo(const FStormSyncConnectionInfo& InConnectionInfo)
 {
-	FAvaSyncConnectionInfo ConnectionInfo;
+	FAvaMediaSyncConnectionInfo ConnectionInfo;
 	ConnectionInfo.EngineVersion = InConnectionInfo.EngineVersion;
 	ConnectionInfo.InstanceId = InConnectionInfo.InstanceId;
 	ConnectionInfo.SessionId = InConnectionInfo.SessionId;
@@ -336,31 +336,31 @@ FAvaSyncConnectionInfo FStormSyncAvaSyncProvider::ConvertConnectionInfo(const FS
 	return ConnectionInfo;
 }
 
-EAvaSyncEngineType FStormSyncAvaSyncProvider::ConvertInstanceType(const EStormSyncEngineType InInstanceType)
+EAvaMediaSyncEngineType FStormSyncAvaSyncProvider::ConvertInstanceType(const EStormSyncEngineType InInstanceType)
 {
-	EAvaSyncEngineType Result;
+	EAvaMediaSyncEngineType Result;
 	switch (InInstanceType)
 	{
 	case EStormSyncEngineType::Server:
-		Result = EAvaSyncEngineType::Server;
+		Result = EAvaMediaSyncEngineType::Server;
 		break;
 	case EStormSyncEngineType::Commandlet:
-		Result = EAvaSyncEngineType::Commandlet;
+		Result = EAvaMediaSyncEngineType::Commandlet;
 		break;
 	case EStormSyncEngineType::Editor:
-		Result = EAvaSyncEngineType::Editor;
+		Result = EAvaMediaSyncEngineType::Editor;
 		break;
 	case EStormSyncEngineType::Game:
-		Result = EAvaSyncEngineType::Game;
+		Result = EAvaMediaSyncEngineType::Game;
 		break;
 	case EStormSyncEngineType::Other:
-		Result = EAvaSyncEngineType::Other;
+		Result = EAvaMediaSyncEngineType::Other;
 		break;
 	case EStormSyncEngineType::Unknown:
-		Result = EAvaSyncEngineType::Unknown;
+		Result = EAvaMediaSyncEngineType::Unknown;
 		break;
 	default:
-		Result = EAvaSyncEngineType::Unknown;
+		Result = EAvaMediaSyncEngineType::Unknown;
 		break;
 	}
 
diff --git a/Engine/Plugins/Experimental/StormSyncAvalancheBridge/Source/StormSyncAvaBridge/Private/ModularFeature/StormSyncAvaSyncProvider.h b/Engine/Plugins/Experimental/StormSyncAvalancheBridge/Source/StormSyncAvaBridge/Private/ModularFeature/StormSyncAvaSyncProvider.h
--- a/Engine/Plugins/Experimental/StormSyncAvalancheBridge/Source/StormSyncAvaBridge/Private/ModularFeature/StormSyncAvaSyncProvider.h
+++ b/Engine/Plugins/Experimental/StormSyncAvalancheBridge/Source/StormSyncAvaBridge/Private/ModularFeature/StormSyncAvaSyncProvider.h
@@ -24,15 +24,15 @@ public:
 	//~ Begin IAvaMediaSyncProvider interface
 	virtual FName GetName() const override;
 	virtual void SyncToAll(const TArray<FName>& InPackageNames) override;
-	virtual void PushToRemote(const FString& InRemoteName, const TArray<FName>& InPackageNames, const FOnAvaSyncResponse& DoneDelegate) override;
-	virtual void PullFromRemote(const FString& InRemoteName, const TArray<FName>& InPackageNames, const FOnAvaSyncResponse& DoneDelegate) override;
-	virtual void CompareWithRemote(const FString& InRemoteName, const TArray<FName>& InPackageNames, const FOnAvaSyncCompareResponse& DoneDelegate) override;
-	virtual FOnAvaSyncPackageModified& GetOnAvaSyncPackageModified() override;
+	virtual void PushToRemote(const FString& InRemoteName, const TArray<FName>& InPackageNames, const FOnAvaMediaSyncResponse& DoneDelegate) override;
+	virtual void PullFromRemote(const FString& InRemoteName, const TArray<FName>& InPackageNames, const FOnAvaMediaSyncResponse& DoneDelegate) override;
+	virtual void CompareWithRemote(const FString& InRemoteName, const TArray<FName>& InPackageNames, const FOnAvaMediaSyncCompareResponse& DoneDelegate) override;
+	virtual FOnAvaMediaSyncPackageModified& GetOnAvaSyncPackageModified() override;
 	//~ End IAvaMediaSyncProvider interface
 
 protected:
 	/** Low lvl API using direct message address destination */
-	virtual void CompareWith(const FMessageAddress& InRemoteAddress, const TArray<FName>& InPackageNames, const FOnAvaSyncCompareResponse& DoneDelegate);
+	virtual void CompareWith(const FMessageAddress& InRemoteAddress, const TArray<FName>& InPackageNames, const FOnAvaMediaSyncCompareResponse& DoneDelegate);
 
 	/**
 	 * Handler called when a file is extracted from an incoming storm sync pak
@@ -41,7 +41,7 @@ protected:
 	void OnPakAssetExtracted(const FName& InPackageName, const FString& InDestFilepath);
 	
 private:
-	FOnAvaSyncPackageModified OnAvaSyncPackageModified;
+	FOnAvaMediaSyncPackageModified OnAvaSyncPackageModified;
 
 	/**
 	 * Helper to get either client or server user data depending on local playback client / server state
@@ -68,14 +68,14 @@ private:
 	static bool IsPlaybackClient();
 
 	/** Helper to create a new error response (used when calling back DoneDelegate to indicate an error) */
-	static TSharedPtr<FAvaSyncCompareResponse> CreateErrorResponse(const FText& InText);
+	static TSharedPtr<FAvaMediaSyncCompareResponse> CreateErrorResponse(const FText& InText);
 
 	/** Helper to convert a FStormSyncTransportSyncResponse struct to the modular feature equivalent */
-	static FAvaSyncResponse ConvertSyncResponse(const TSharedPtr<FStormSyncTransportSyncResponse>& InResponse);
+	static FAvaMediaSyncResponse ConvertSyncResponse(const TSharedPtr<FStormSyncTransportSyncResponse>& InResponse);
 	
 	/** Helper to convert a FStormSyncConnectionInfo struct to the modular feature equivalent */
-	static FAvaSyncConnectionInfo ConvertConnectionInfo(const FStormSyncConnectionInfo& InConnectionInfo);
+	static FAvaMediaSyncConnectionInfo ConvertConnectionInfo(const FStormSyncConnectionInfo& InConnectionInfo);
 
 	/** Helper to convert instance type enum specific to this implementation to the interface one (for connection info) */
-	static EAvaSyncEngineType ConvertInstanceType(const EStormSyncEngineType InInstanceType);
+	static EAvaMediaSyncEngineType ConvertInstanceType(const EStormSyncEngineType InInstanceType);
 };
diff --git a/Engine/Plugins/Experimental/StormSyncAvalancheBridge/Source/StormSyncAvaBridge/Private/StormSyncAvaBridge.cpp b/Engine/Plugins/Experimental/StormSyncAvalancheBridge/Source/StormSyncAvaBridge/Private/StormSyncAvaBridge.cpp
--- a/Engine/Plugins/Experimental/StormSyncAvalancheBridge/Source/StormSyncAvaBridge/Private/StormSyncAvaBridge.cpp
+++ b/Engine/Plugins/Experimental/StormSyncAvalancheBridge/Source/StormSyncAvaBridge/Private/StormSyncAvaBridge.cpp
@@ -8,8 +8,8 @@
 #include "IStormSyncTransportServerModule.h"
 #include "Misc/CoreDelegates.h"
 #include "ModularFeature/StormSyncAvaSyncProvider.h"
-#include "Playback/AvaMediaPlaybackServer.h"
-#include "Playback/IAvaMediaPlaybackClient.h"
+#include "Playback/AvaPlaybackServer.h"
+#include "Playback/IAvaPlaybackClient.h"
 #include "StormSyncAvaBridgeCommon.h"
 #include "StormSyncAvaBridgeLog.h"
 #include "StormSyncAvaBridgeUtils.h"
@@ -58,8 +58,8 @@ void FStormSyncAvaBridgeModule::OnPostEngineInit()
 			RegisterUserDataForPlaybackClient();
 		}
 
-		IAvaMediaModule::Get().GetOnAvaMediaPlaybackServerStarted().AddStatic(&FStormSyncAvaBridgeModule::RegisterUserDataForPlaybackServer);
-		IAvaMediaModule::Get().GetOnAvaMediaPlaybackClientStarted().AddStatic(&FStormSyncAvaBridgeModule::RegisterUserDataForPlaybackClient);
+		IAvaMediaModule::Get().GetOnAvaPlaybackServerStarted().AddStatic(&FStormSyncAvaBridgeModule::RegisterUserDataForPlaybackServer);
+		IAvaMediaModule::Get().GetOnAvaPlaybackClientStarted().AddStatic(&FStormSyncAvaBridgeModule::RegisterUserDataForPlaybackClient);
 	}
 
 	if (IStormSyncTransportServerModule::IsAvailable())
@@ -133,7 +133,7 @@ void FStormSyncAvaBridgeModule::ExecuteGetUserData(const TArray<FString>& Args)
 		return;
 	}
 
-	const IAvaMediaPlaybackClient& PlaybackClient = IAvaMediaModule::Get().GetMediaPlaybackClient();
+	const IAvaPlaybackClient& PlaybackClient = IAvaMediaModule::Get().GetMediaPlaybackClient();
 	for (const FString& ServerName : ServerNames)
 	{
 		FString ServerValue = PlaybackClient.GetServerUserData(ServerName, Key);
@@ -161,7 +161,7 @@ void FStormSyncAvaBridgeModule::RegisterUserDataForPlaybackServer()
 		return;
 	}
 
-	if (const TSharedPtr<FAvaMediaPlaybackServer> PlaybackServer = IAvaMediaModule::Get().GetMediaPlaybackServer())
+	if (const TSharedPtr<FAvaPlaybackServer> PlaybackServer = IAvaMediaModule::Get().GetMediaPlaybackServer())
 	{
 		const FString StormSyncServerAddress = IStormSyncTransportServerModule::Get().GetServerEndpointMessageAddressId();
 		const FString DiscoveryManagerAddress = IStormSyncTransportServerModule::Get().GetDiscoveryManagerMessageAddressId();
@@ -187,7 +187,7 @@ void FStormSyncAvaBridgeModule::RegisterUserDataForPlaybackClient()
 
 	if (IAvaMediaModule::Get().IsMediaPlaybackClientStarted())
 	{
-		IAvaMediaPlaybackClient& PlaybackClient = IAvaMediaModule::Get().GetMediaPlaybackClient();
+		IAvaPlaybackClient& PlaybackClient = IAvaMediaModule::Get().GetMediaPlaybackClient();
 
 		const FString StormSyncServerAddress = IStormSyncTransportServerModule::Get().GetServerEndpointMessageAddressId();
 		const FString DiscoveryManagerAddress = IStormSyncTransportServerModule::Get().GetDiscoveryManagerMessageAddressId();
diff --git a/Engine/Plugins/Experimental/StormSyncAvalancheBridge/Source/StormSyncAvaBridge/Private/StormSyncAvaBridgeUtils.cpp b/Engine/Plugins/Experimental/StormSyncAvalancheBridge/Source/StormSyncAvaBridge/Private/StormSyncAvaBridgeUtils.cpp
--- a/Engine/Plugins/Experimental/StormSyncAvalancheBridge/Source/StormSyncAvaBridge/Private/StormSyncAvaBridgeUtils.cpp
+++ b/Engine/Plugins/Experimental/StormSyncAvalancheBridge/Source/StormSyncAvaBridge/Private/StormSyncAvaBridgeUtils.cpp
@@ -1,23 +1,21 @@
 // Copyright Epic Games, Inc. All Rights Reserved.
 
-
 #include "StormSyncAvaBridgeUtils.h"
-
-#include "AvalancheBroadcast.h"
+#include "Broadcast/AvaBroadcast.h"
 
 TArray<FString> FStormSyncAvaBridgeUtils::GetServerNamesForChannel(const FString& InChannelName)
 {
 	TArray<FString> ServerNames;
 	
-	if (UAvalancheBroadcast* Broadcast = UAvalancheBroadcast::GetAvalancheBroadcast())
+	if (UAvaBroadcast* Broadcast = UAvaBroadcast::GetBroadcast())
 	{
-		const FAvaOutputChannel Channel = Broadcast->GetCurrentProfile().GetChannel(FName(*InChannelName));
+		const FAvaBroadcastOutputChannel Channel = Broadcast->GetCurrentProfile().GetChannel(FName(*InChannelName));
 		if (Channel.IsValidChannel())
 		{
 			TArray<UMediaOutput*> Outputs = Channel.GetMediaOutputs();
 			for (const UMediaOutput* Output : Outputs)
 			{
-				FAvaMediaOutputInfo OutputInfo = Channel.GetMediaOutputInfo(Output);
+				FAvaBroadcastMediaOutputInfo OutputInfo = Channel.GetMediaOutputInfo(Output);
 				if (OutputInfo.IsValid() && OutputInfo.IsRemote())
 				{
 					ServerNames.Add(OutputInfo.ServerName);
diff --git a/Engine/Plugins/Experimental/StormSyncAvalancheBridge/Source/StormSyncAvaBridgeEditor/Private/StormSyncAvaPlaylistExtender.cpp b/Engine/Plugins/Experimental/StormSyncAvalancheBridge/Source/StormSyncAvaBridgeEditor/Private/StormSyncAvaPlaylistExtender.cpp
--- a/Engine/Plugins/Experimental/StormSyncAvalancheBridge/Source/StormSyncAvaBridgeEditor/Private/StormSyncAvaPlaylistExtender.cpp
+++ b/Engine/Plugins/Experimental/StormSyncAvalancheBridge/Source/StormSyncAvaBridgeEditor/Private/StormSyncAvaPlaylistExtender.cpp
@@ -1,15 +1,13 @@
 // Copyright Epic Games, Inc. All Rights Reserved.
 
 #include "StormSyncAvaPlaylistExtender.h"
-
 #include "Framework/MultiBox/MultiBoxBuilder.h"
 #include "IAvaMediaEditorModule.h"
 #include "IAvaMediaModule.h"
 #include "IStormSyncTransportClientModule.h"
-#include "Playback/AvaMediaPlaybackServer.h"
-#include "Playback/IAvaMediaPlaybackClient.h"
-#include "Playlist/AvaPlaylistEditor.h"
-#include "Playlist/AvalanchePlaylist.h"
+#include "Playback/IAvaPlaybackClient.h"
+#include "Rundown/AvaRundown.h"
+#include "Rundown/AvaRundownEditor.h"
 #include "StormSyncAvaBridgeCommon.h"
 #include "StormSyncAvaBridgeEditorLog.h"
 #include "StormSyncAvaBridgeUtils.h"
@@ -30,8 +28,8 @@ FStormSyncAvaPlaylistExtender::~FStormSyncAvaPlaylistExtender()
 
 	if (IAvaMediaEditorModule::IsLoaded())
 	{
-		UnregisterExtension(IAvaMediaEditorModule::Get().GetPlaylistMenuExtensibilityManager(), MenuExtenderHandle);
-		UnregisterExtension(IAvaMediaEditorModule::Get().GetPlaylistToolBarExtensibilityManager(), ToolbarExtenderHandle);
+		UnregisterExtension(IAvaMediaEditorModule::Get().GetRundownMenuExtensibilityManager(), MenuExtenderHandle);
+		UnregisterExtension(IAvaMediaEditorModule::Get().GetRundownToolBarExtensibilityManager(), ToolbarExtenderHandle);
 	}
 }
 
@@ -47,13 +45,13 @@ void FStormSyncAvaPlaylistExtender::RegisterMenuExtensions()
 {
 	// Context menu extension
 	MenuExtenderHandle = RegisterExtension(
-		IAvaMediaEditorModule::Get().GetPlaylistMenuExtensibilityManager(),
+		IAvaMediaEditorModule::Get().GetRundownMenuExtensibilityManager(),
 		FAssetEditorExtender::CreateSP(this, &FStormSyncAvaPlaylistExtender::AddMenuExtender)
 	);
 	
 	// Toolbar extension
 	ToolbarExtenderHandle = RegisterExtension(
-		IAvaMediaEditorModule::Get().GetPlaylistToolBarExtensibilityManager(),
+		IAvaMediaEditorModule::Get().GetRundownToolBarExtensibilityManager(),
 		FAssetEditorExtender::CreateSP(this, &FStormSyncAvaPlaylistExtender::AddToolbarExtender)
 	);
 }
@@ -88,7 +86,7 @@ TSharedRef<FExtender> FStormSyncAvaPlaylistExtender::AddMenuExtender(const TShar
 	STORM_SYNC_AVA_EDITOR_LOG(Display, TEXT("FStormSyncAvaPlaylistExtender::AddMenuExtender - Adding in menu extensions"))
 	TSharedRef<FExtender> Extender(new FExtender());
 
-	const UAvalanchePlaylist* Playlist = ContextSensitiveObjects.IsValidIndex(0) ? Cast<UAvalanchePlaylist>(ContextSensitiveObjects[0]) : nullptr;
+	const UAvaRundown* Playlist = ContextSensitiveObjects.IsValidIndex(0) ? Cast<UAvaRundown>(ContextSensitiveObjects[0]) : nullptr;
 	if (!Playlist)
 	{
 		return Extender;
@@ -100,7 +98,7 @@ TSharedRef<FExtender> FStormSyncAvaPlaylistExtender::AddMenuExtender(const TShar
 		return Extender;
 	}
 
-	const TSharedPtr<FAvaPlaylistEditor> PlaylistEditor = StaticCastSharedPtr<FAvaPlaylistEditor>(AssetEditor);
+	const TSharedPtr<FAvaRundownEditor> PlaylistEditor = StaticCastSharedPtr<FAvaRundownEditor>(AssetEditor);
 	if (!PlaylistEditor.IsValid())
 	{
 		return Extender;
@@ -112,17 +110,17 @@ TSharedRef<FExtender> FStormSyncAvaPlaylistExtender::AddMenuExtender(const TShar
 		EExtensionHook::After,
 		InCommandList,
 		// Convert to weak ptr to prevent ownership to the playlist editor and potentially increasing its lifetime
-		FMenuExtensionDelegate::CreateSP(this, &FStormSyncAvaPlaylistExtender::CreateTemplateContextMenu, Playlist, TWeakPtr<FAvaPlaylistEditor>(PlaylistEditor))
+		FMenuExtensionDelegate::CreateSP(this, &FStormSyncAvaPlaylistExtender::CreateTemplateContextMenu, Playlist, TWeakPtr<FAvaRundownEditor>(PlaylistEditor))
 	);
 
 	return Extender;
 }
 
-void FStormSyncAvaPlaylistExtender::CreateTemplateContextMenu(FMenuBuilder& MenuBuilder, const UAvalanchePlaylist* InPlaylist, TWeakPtr<FAvaPlaylistEditor> InPlaylistEditor)
+void FStormSyncAvaPlaylistExtender::CreateTemplateContextMenu(FMenuBuilder& MenuBuilder, const UAvaRundown* InPlaylist, TWeakPtr<FAvaRundownEditor> InPlaylistEditor)
 {
 	check(InPlaylist);
 
-	const TSharedPtr<FAvaPlaylistEditor> PlaylistEditor = InPlaylistEditor.Pin();
+	const TSharedPtr<FAvaRundownEditor> PlaylistEditor = InPlaylistEditor.Pin();
 	if (!PlaylistEditor.IsValid())
 	{
 		STORM_SYNC_AVA_EDITOR_LOG(Error, TEXT("FStormSyncAvaPlaylistExtender::CreateTemplateContextMenu - Invalid shared ptr from weak ptr delegate param"))
@@ -152,7 +150,7 @@ void FStormSyncAvaPlaylistExtender::CreateTemplateContextMenu(FMenuBuilder& Menu
 			FText::Format(LOCTEXT("Initialize_Tooltip", "Sync asset over remote node.{0}"), DisabledTooltipReason) :
 			FText::Format(LOCTEXT("Initialize_Tooltip_Invalid", "Please ensure the playlist page is using a valid Avalanche Blueprint.{0}"), DisabledTooltipReason);
 
-		const TArray<FAvalanchePage> SelectedTemplatePages = GetSelectedPages(InPlaylist, InPlaylistEditor);
+		const TArray<FAvaRundownPage> SelectedTemplatePages = GetSelectedPages(InPlaylist, InPlaylistEditor);
 		
 		MenuBuilder.AddMenuEntry(
 			LOCTEXT("Initialize_Label", "Initialize"),
@@ -248,15 +246,15 @@ void FStormSyncAvaPlaylistExtender::CreateTemplateContextMenu(FMenuBuilder& Menu
 	MenuBuilder.EndSection();
 }
 
-void FStormSyncAvaPlaylistExtender::HandleInitializeAction(const UAvalanchePlaylist* InPlaylist, TArray<FAvalanchePage> InSelectedTemplatePages)
+void FStormSyncAvaPlaylistExtender::HandleInitializeAction(const UAvaRundown* InPlaylist, TArray<FAvaRundownPage> InSelectedTemplatePages)
 {
 	// Build a list of package names to push grouped by channel name (and remote)
 	// Key is the remote address id, value is the list of package names to synchronize
 	TMap<FString, TArray<FName>> PackageNamesPerChannel;
 	
-	for (const FAvalanchePage& SelectedPage : InSelectedTemplatePages)
+	for (const FAvaRundownPage& SelectedPage : InSelectedTemplatePages)
 	{
-		const FString AssetName = SelectedPage.GetAvalancheAssetPath(InPlaylist).GetLongPackageName();
+		const FString AssetName = SelectedPage.GetAssetPath(InPlaylist).GetLongPackageName();
 		TArray<FString> ChannelNames = GetChannelNamesForTemplatePage(InPlaylist, SelectedPage);
 
 		// From the list of channel names that match this asset to sync, build the list of server names
@@ -279,7 +277,7 @@ void FStormSyncAvaPlaylistExtender::HandleInitializeAction(const UAvalanchePlayl
 
 		if (IAvaMediaModule::Get().IsMediaPlaybackClientStarted())
 		{
-			const IAvaMediaPlaybackClient& PlaybackClient = IAvaMediaModule::Get().GetMediaPlaybackClient();
+			const IAvaPlaybackClient& PlaybackClient = IAvaMediaModule::Get().GetMediaPlaybackClient();
 			for (const FString& ServerName : RemoteServerNames)
 			{
 				// Get address id for storm sync client on playback host
@@ -318,7 +316,7 @@ TSharedRef<FExtender> FStormSyncAvaPlaylistExtender::AddToolbarExtender(const TS
 	STORM_SYNC_AVA_EDITOR_LOG(Display, TEXT("FStormSyncAvaPlaylistExtender::AddToolbarExtender - Adding in toolbar extensions"))
 	TSharedRef<FExtender> Extender(new FExtender());
 
-	const UAvalanchePlaylist* Playlist = ContextSensitiveObjects.IsValidIndex(0) ? Cast<UAvalanchePlaylist>(ContextSensitiveObjects[0]) : nullptr;
+	const UAvaRundown* Playlist = ContextSensitiveObjects.IsValidIndex(0) ? Cast<UAvaRundown>(ContextSensitiveObjects[0]) : nullptr;
 	if (!Playlist)
 	{
 		return Extender;
@@ -330,7 +328,7 @@ TSharedRef<FExtender> FStormSyncAvaPlaylistExtender::AddToolbarExtender(const TS
 		return Extender;
 	}
 
-	const TSharedPtr<FAvaPlaylistEditor> PlaylistEditor = StaticCastSharedPtr<FAvaPlaylistEditor>(AssetEditor);
+	const TSharedPtr<FAvaRundownEditor> PlaylistEditor = StaticCastSharedPtr<FAvaRundownEditor>(AssetEditor);
 	if (!PlaylistEditor.IsValid())
 	{
 		return Extender;
@@ -340,15 +338,15 @@ TSharedRef<FExtender> FStormSyncAvaPlaylistExtender::AddToolbarExtender(const TS
 		"Pages",
 		EExtensionHook::After,
 		InCommandList,
-		FToolBarExtensionDelegate::CreateSP(this, &FStormSyncAvaPlaylistExtender::FillToolbar, TWeakPtr<FAvaPlaylistEditor>(PlaylistEditor))
+		FToolBarExtensionDelegate::CreateSP(this, &FStormSyncAvaPlaylistExtender::FillToolbar, TWeakPtr<FAvaRundownEditor>(PlaylistEditor))
 	);
 	
-	STORM_SYNC_AVA_EDITOR_LOG(Display, TEXT("FStormSyncAvaPlaylistExtender::AddToolbarExtender - Playlist: %s"), *GetNameSafe(Playlist))
+	STORM_SYNC_AVA_EDITOR_LOG(Display, TEXT("FStormSyncAvaPlaylistExtender::AddToolbarExtender - Rundown: %s"), *GetNameSafe(Playlist))
 	
 	return Extender;
 }
 
-void FStormSyncAvaPlaylistExtender::FillToolbar(FToolBarBuilder& ToolbarBuilder, TWeakPtr<FAvaPlaylistEditor> InPlaylistEditor)
+void FStormSyncAvaPlaylistExtender::FillToolbar(FToolBarBuilder& ToolbarBuilder, TWeakPtr<FAvaRundownEditor> InPlaylistEditor)
 {
 	ToolbarBuilder.BeginSection(TEXT("StormSync"));
 	{
@@ -364,21 +362,21 @@ void FStormSyncAvaPlaylistExtender::FillToolbar(FToolBarBuilder& ToolbarBuilder,
 	ToolbarBuilder.EndSection();
 }
 
-TSharedRef<SWidget> FStormSyncAvaPlaylistExtender::GenerateToolbarMenu(TWeakPtr<FAvaPlaylistEditor> InPlaylistEditor)
+TSharedRef<SWidget> FStormSyncAvaPlaylistExtender::GenerateToolbarMenu(TWeakPtr<FAvaRundownEditor> InPlaylistEditor)
 {
 	TArray<FName> PackageNames;
 
-	const TSharedPtr<FAvaPlaylistEditor> PlaylistEditor = InPlaylistEditor.Pin();
+	const TSharedPtr<FAvaRundownEditor> PlaylistEditor = InPlaylistEditor.Pin();
 
-	if (PlaylistEditor.IsValid() && PlaylistEditor->IsPlaylistValid())
+	if (PlaylistEditor.IsValid() && PlaylistEditor->IsRundownValid())
 	{
-		const UAvalanchePlaylist* Playlist = PlaylistEditor->GetPlaylist(); 
+		const UAvaRundown* Playlist = PlaylistEditor->GetRundown(); 
 
-		// Gather the list of all package names from Avalanche Blueprints in this Playlist pages
-		const FAvalanchePageCollection& AvalanchePageCollection = Playlist->GetTemplatePages();
-		for (const FAvalanchePage& AvalanchePage : AvalanchePageCollection.Pages)
+		// Gather the list of all package names from Avalanche Blueprints in this Rundown pages
+		const FAvaRundownPageCollection& AvalanchePageCollection = Playlist->GetTemplatePages();
+		for (const FAvaRundownPage& AvalanchePage : AvalanchePageCollection.Pages)
 		{
-			if (FString PackageName = AvalanchePage.GetAvalancheAssetPath(Playlist).GetLongPackageName(); !PackageName.IsEmpty())
+			if (FString PackageName = AvalanchePage.GetAssetPath(Playlist).GetLongPackageName(); !PackageName.IsEmpty())
 			{
 				PackageNames.AddUnique(FName(*PackageName));
 			}
@@ -410,7 +408,7 @@ TSharedRef<SWidget> FStormSyncAvaPlaylistExtender::GenerateToolbarMenu(TWeakPtr<
 		const bool bIsPushEnabled = !PackageNames.IsEmpty();
 		if (!bIsPushEnabled)
 		{
-			LabelText = LOCTEXT("PushAssetsToolbarMenuEntryInvalid", "Cannot push. Playlist pages have no valid assets.");
+			LabelText = LOCTEXT("PushAssetsToolbarMenuEntryInvalid", "Cannot push. Rundown pages have no valid assets.");
 			TooltipText = LOCTEXT("PushAssetsToolbarMenuEntryTooltipInvalid", "Please ensure the playlist pages are using a valid Avalanche Blueprint");
 		}
 
@@ -483,27 +481,27 @@ void FStormSyncAvaPlaylistExtender::PushPackagesToRemote(const FString& RemoteAd
 	IStormSyncTransportClientModule::Get().PushPackages(PackageDescriptor, InPackageNames, RemoteMessageAddress, Delegate);
 }
 
-TArray<FName> FStormSyncAvaPlaylistExtender::GetSelectedPackagesNames(const UAvalanchePlaylist* InPlaylist, const TWeakPtr<FAvaPlaylistEditor>& InPlaylistEditor)
+TArray<FName> FStormSyncAvaPlaylistExtender::GetSelectedPackagesNames(const UAvaRundown* InPlaylist, const TWeakPtr<FAvaRundownEditor>& InPlaylistEditor)
 {
 	check(InPlaylist);
 	
 	TArray<FName> Result;
 
-	const TArray<FAvalanchePage> SelectedPages = GetSelectedPages(InPlaylist, InPlaylistEditor);
-	Algo::Transform(SelectedPages, Result, [InPlaylist](const FAvalanchePage& Page)
+	const TArray<FAvaRundownPage> SelectedPages = GetSelectedPages(InPlaylist, InPlaylistEditor);
+	Algo::Transform(SelectedPages, Result, [InPlaylist](const FAvaRundownPage& Page)
 	{
-		return Page.GetAvalancheAssetPath(InPlaylist).GetLongPackageFName();
+		return Page.GetAssetPath(InPlaylist).GetLongPackageFName();
 	});
 
 	return Result;
 }
 
-TArray<FAvalanchePage> FStormSyncAvaPlaylistExtender::GetSelectedPages(const UAvalanchePlaylist* InPlaylist, const TWeakPtr<FAvaPlaylistEditor>& InPlaylistEditor)
+TArray<FAvaRundownPage> FStormSyncAvaPlaylistExtender::GetSelectedPages(const UAvaRundown* InPlaylist, const TWeakPtr<FAvaRundownEditor>& InPlaylistEditor)
 {
 	check(InPlaylist);
 	
-	TArray<FAvalanchePage> Result;
-	const TSharedPtr<FAvaPlaylistEditor> PlaylistEditor = InPlaylistEditor.Pin();
+	TArray<FAvaRundownPage> Result;
+	const TSharedPtr<FAvaRundownEditor> PlaylistEditor = InPlaylistEditor.Pin();
 	if (!PlaylistEditor.IsValid())
 	{
 		STORM_SYNC_AVA_EDITOR_LOG(Display, TEXT("FStormSyncAvaPlaylistExtender::GetSelectedPages - Invalid shared ptr from weak ptr delegate param"))
@@ -513,13 +511,13 @@ TArray<FAvalanchePage> FStormSyncAvaPlaylistExtender::GetSelectedPages(const UAv
 	const TConstArrayView<int32> SelectedPageIds = PlaylistEditor->GetSelectedPagesOnFocusedWidget();
 	for (const int32 SelectedPageId : SelectedPageIds)
 	{
-		FAvalanchePage AvalanchePage = InPlaylist->GetPage(SelectedPageId);
+		FAvaRundownPage AvalanchePage = InPlaylist->GetPage(SelectedPageId);
 		if (!AvalanchePage.IsValidPage())
 		{
 			continue;
 		}
 
-		FSoftObjectPath SoftAvalancheAssetPath = AvalanchePage.GetAvalancheAssetPath(InPlaylist);
+		FSoftObjectPath SoftAvalancheAssetPath = AvalanchePage.GetAssetPath(InPlaylist);
 		FString LongPackageName = SoftAvalancheAssetPath.GetLongPackageName();
 		FString AssetName = SoftAvalancheAssetPath.GetAssetName();
 
@@ -532,7 +530,7 @@ TArray<FAvalanchePage> FStormSyncAvaPlaylistExtender::GetSelectedPages(const UAv
 	return Result;
 }
 
-TArray<FString> FStormSyncAvaPlaylistExtender::GetChannelNamesForTemplatePage(const UAvalanchePlaylist* InPlaylist, const FAvalanchePage& InTemplatePage)
+TArray<FString> FStormSyncAvaPlaylistExtender::GetChannelNamesForTemplatePage(const UAvaRundown* InPlaylist, const FAvaRundownPage& InTemplatePage)
 {
 	check(InPlaylist);
 
@@ -540,16 +538,16 @@ TArray<FString> FStormSyncAvaPlaylistExtender::GetChannelNamesForTemplatePage(co
 
 	// Here, we try to determine the list of channels to consider for a sync operation, from the selected template page,
 	// with instanced pages that are using the selected package name (Avalanche Blueprint)
-	const FString PackageName = InTemplatePage.GetAvalancheAssetPath(InPlaylist).GetLongPackageName();
+	const FString PackageName = InTemplatePage.GetAssetPath(InPlaylist).GetLongPackageName();
 
 	// Build the list of instanced pages matching the asset name we want to sync
-	TArray<FAvalanchePage> Pages = InPlaylist->GetInstancedPages().Pages.FilterByPredicate([InPlaylist, PackageName](const FAvalanchePage& Page)
+	TArray<FAvaRundownPage> Pages = InPlaylist->GetInstancedPages().Pages.FilterByPredicate([InPlaylist, PackageName](const FAvaRundownPage& Page)
 	{
-		return Page.GetAvalancheAssetPath(InPlaylist).GetLongPackageName() == PackageName;
+		return Page.GetAssetPath(InPlaylist).GetLongPackageName() == PackageName;
 	});
 
 	// From there, build a unique list of channel outputs
-	for (const FAvalanchePage& AvalanchePage : Pages)
+	for (const FAvaRundownPage& AvalanchePage : Pages)
 	{
 		ChannelNames.AddUnique(AvalanchePage.GetChannelName().ToString());
 	}
@@ -557,7 +555,7 @@ TArray<FString> FStormSyncAvaPlaylistExtender::GetChannelNamesForTemplatePage(co
 	return ChannelNames;
 }
 
-bool FStormSyncAvaPlaylistExtender::GetContextMenuSelectionInfos(const UAvalanchePlaylist* InPlaylist, const TWeakPtr<FAvaPlaylistEditor>& InPlaylistEditor, bool& bOutIsValidSelection, FText& OutDisabledReasonTooltip, TArray<FName>& OutSelectedPackageNames, FCanExecuteAction& OutCanExecuteAction)
+bool FStormSyncAvaPlaylistExtender::GetContextMenuSelectionInfos(const UAvaRundown* InPlaylist, const TWeakPtr<FAvaRundownEditor>& InPlaylistEditor, bool& bOutIsValidSelection, FText& OutDisabledReasonTooltip, TArray<FName>& OutSelectedPackageNames, FCanExecuteAction& OutCanExecuteAction)
 {
 	if (!FModuleManager::Get().IsModuleLoaded(TEXT("StormSyncEditor")))
 	{
diff --git a/Engine/Plugins/Experimental/StormSyncAvalancheBridge/Source/StormSyncAvaBridgeEditor/Private/StormSyncAvaPlaylistExtender.h b/Engine/Plugins/Experimental/StormSyncAvalancheBridge/Source/StormSyncAvaBridgeEditor/Private/StormSyncAvaPlaylistExtender.h
--- a/Engine/Plugins/Experimental/StormSyncAvalancheBridge/Source/StormSyncAvaBridgeEditor/Private/StormSyncAvaPlaylistExtender.h
+++ b/Engine/Plugins/Experimental/StormSyncAvalancheBridge/Source/StormSyncAvaBridgeEditor/Private/StormSyncAvaPlaylistExtender.h
@@ -6,12 +6,12 @@
 #include "MessageEndpoint.h"
 #include "Toolkits/AssetEditorToolkit.h"
 
-class FAvaPlaylistEditor;
+class FAvaRundownEditor;
 class FExtensibilityManager;
-class UAvalanchePlaylist;
-struct FAvalanchePage;
+class UAvaRundown;
+struct FAvaRundownPage;
 
-/** This class handles Avalanche Playlist assets UI extensions for context menu / toolbars */
+/** This class handles Avalanche Rundown assets UI extensions for context menu / toolbars */
 class FStormSyncAvaPlaylistExtender : public TSharedFromThis<FStormSyncAvaPlaylistExtender>
 {
 public:
@@ -47,35 +47,35 @@ private:
 	TSharedRef<FExtender> AddMenuExtender(const TSharedRef<FUICommandList> InCommandList, const TArray<UObject*> ContextSensitiveObjects);
 
 	/* UI Menu Extension handler for template panel */
-	void CreateTemplateContextMenu(FMenuBuilder& MenuBuilder, const UAvalanchePlaylist* InPlaylist, TWeakPtr<FAvaPlaylistEditor> InPlaylistEditor);
+	void CreateTemplateContextMenu(FMenuBuilder& MenuBuilder, const UAvaRundown* InPlaylist, TWeakPtr<FAvaRundownEditor> InPlaylistEditor);
 
 	/** Context menu handler for initialize action */
-	void HandleInitializeAction(const UAvalanchePlaylist* InPlaylist, TArray<FAvalanchePage> InSelectedTemplatePages);
+	void HandleInitializeAction(const UAvaRundown* InPlaylist, TArray<FAvaRundownPage> InSelectedTemplatePages);
 	
 	/** Gets the extender to use for playlist context sensitive menus */
 	TSharedRef<FExtender> AddToolbarExtender(const TSharedRef<FUICommandList> InCommandList, const TArray<UObject*> ContextSensitiveObjects);
 	
 	/** Construct toolbar widgets for playlist sync actions */
-	void FillToolbar(FToolBarBuilder& ToolbarBuilder, TWeakPtr<FAvaPlaylistEditor> InPlaylistEditor);
+	void FillToolbar(FToolBarBuilder& ToolbarBuilder, TWeakPtr<FAvaRundownEditor> InPlaylistEditor);
 
 	/** Creates widget for toolbar content */
-	TSharedRef<SWidget> GenerateToolbarMenu(TWeakPtr<FAvaPlaylistEditor> InPlaylistEditor);
+	TSharedRef<SWidget> GenerateToolbarMenu(TWeakPtr<FAvaRundownEditor> InPlaylistEditor);
 
 	/** Storm sync push a list of package names to a remote address ID */
 	static void PushPackagesToRemote(const FString& RemoteAddressId, const TArray<FName>& InPackageNames);
 
 	/** Gather a list of package names from currently selected pages in editor */
-	static TArray<FName> GetSelectedPackagesNames(const UAvalanchePlaylist* InPlaylist, const TWeakPtr<FAvaPlaylistEditor>& InPlaylistEditor);
+	static TArray<FName> GetSelectedPackagesNames(const UAvaRundown* InPlaylist, const TWeakPtr<FAvaRundownEditor>& InPlaylistEditor);
 	
 	/**
 	 * Returns the list of currently selected pages in editor
 	 *
 	 * It will only return selection if the page is referencing a valid Avalanche Blueprint (not "None" path)
 	 */
-	static TArray<FAvalanchePage> GetSelectedPages(const UAvalanchePlaylist* InPlaylist, const TWeakPtr<FAvaPlaylistEditor>& InPlaylistEditor);
+	static TArray<FAvaRundownPage> GetSelectedPages(const UAvaRundown* InPlaylist, const TWeakPtr<FAvaRundownEditor>& InPlaylistEditor);
 
 	/** Returns a unique list of channel names, gather from instanced pages, that are matching the passed in template page selection and selected asset name */
-	static TArray<FString> GetChannelNamesForTemplatePage(const UAvalanchePlaylist* InPlaylist, const FAvalanchePage& InTemplatePage);
+	static TArray<FString> GetChannelNamesForTemplatePage(const UAvaRundown* InPlaylist, const FAvaRundownPage& InTemplatePage);
 
 	/**
 	 * Provides a way for context menu builders to gather information about current selection.
@@ -90,8 +90,8 @@ private:
 	 * @return false if StormSyncEditor module is not available, and we failed to determine selection state
 	 */
 	static bool GetContextMenuSelectionInfos(
-		const UAvalanchePlaylist* InPlaylist,
-		const TWeakPtr<FAvaPlaylistEditor>& InPlaylistEditor,
+		const UAvaRundown* InPlaylist,
+		const TWeakPtr<FAvaRundownEditor>& InPlaylistEditor,
 		bool& bOutIsValidSelection,
 		FText& OutDisabledReasonTooltip,
 		TArray<FName>& OutSelectedPackageNames,
