Storm Sync Motion Design Bridge:
- General style formatting

diff --git a/Engine/Plugins/Experimental/StormSyncAvalancheBridge/Source/StormSyncAvaBridge/Private/ModularFeature/StormSyncAvaSyncProvider.cpp b/Engine/Plugins/Experimental/StormSyncAvalancheBridge/Source/StormSyncAvaBridge/Private/ModularFeature/StormSyncAvaSyncProvider.cpp
--- a/Engine/Plugins/Experimental/StormSyncAvalancheBridge/Source/StormSyncAvaBridge/Private/ModularFeature/StormSyncAvaSyncProvider.cpp
+++ b/Engine/Plugins/Experimental/StormSyncAvalancheBridge/Source/StormSyncAvaBridge/Private/ModularFeature/StormSyncAvaSyncProvider.cpp
@@ -29,7 +29,7 @@ FName FStormSyncAvaSyncProvider::GetName() const
 
 void FStormSyncAvaSyncProvider::SyncToAll(const TArray<FName>& InPackageNames)
 {
-	STORM_SYNC_AVA_LOG(Display, TEXT("FStormSyncAvaSyncProvider::SyncToAll InPackageNames: %d"), InPackageNames.Num());
+	UE_LOG(LogStormSyncAvaBridge, Display, TEXT("FStormSyncAvaSyncProvider::SyncToAll InPackageNames: %d"), InPackageNames.Num());
 
 	const FStormSyncPackageDescriptor PackageDescriptor;
 	IStormSyncTransportClientModule::Get().SynchronizePackages(PackageDescriptor, InPackageNames);
@@ -37,14 +37,14 @@ void FStormSyncAvaSyncProvider::SyncToAll(const TArray<FName>& InPackageNames)
 
 void FStormSyncAvaSyncProvider::PushToRemote(const FString& InRemoteName, const TArray<FName>& InPackageNames, const FOnAvaMediaSyncResponse& DoneDelegate)
 {
-	STORM_SYNC_AVA_LOG(Display, TEXT("FStormSyncAvaSyncProvider::PushToRemote InRemoteName: %s, InPackageNames: %d"), *InRemoteName, InPackageNames.Num());
+	UE_LOG(LogStormSyncAvaBridge, Display, TEXT("FStormSyncAvaSyncProvider::PushToRemote InRemoteName: %s, InPackageNames: %d"), *InRemoteName, InPackageNames.Num());
 	
 	FText ErrorText;
 	FMessageAddress MessageAddress;
 	// Get address from either client or server user data, depending on local client / server state
 	if (!GetAddressFromUserData(InRemoteName, UE::StormSync::AvaBridgeCommon::StormSyncClientAddressKey, MessageAddress, &ErrorText))
 	{
-		STORM_SYNC_AVA_LOG(Error, TEXT("FStormSyncAvaSyncProvider::PushToRemote - %s"), *ErrorText.ToString());
+		UE_LOG(LogStormSyncAvaBridge, Error, TEXT("FStormSyncAvaSyncProvider::PushToRemote - %s"), *ErrorText.ToString());
 		DoneDelegate.ExecuteIfBound(CreateErrorResponse(ErrorText));
 		return;
 	}
@@ -57,7 +57,7 @@ void FStormSyncAvaSyncProvider::PushToRemote(const FString& InRemoteName, const
 			return;
 		}
 		
-		STORM_SYNC_AVA_LOG(Display, TEXT("FStormSyncAvaSyncProvider::PushToRemote - Response: %s"), *Response->ToString());
+		UE_LOG(LogStormSyncAvaBridge, Display, TEXT("FStormSyncAvaSyncProvider::PushToRemote - Response: %s"), *Response->ToString());
 
 		if (DoneDelegate.IsBound())
 		{
@@ -72,14 +72,14 @@ void FStormSyncAvaSyncProvider::PushToRemote(const FString& InRemoteName, const
 
 void FStormSyncAvaSyncProvider::PullFromRemote(const FString& InRemoteName, const TArray<FName>& InPackageNames, const FOnAvaMediaSyncResponse& DoneDelegate)
 {
-	STORM_SYNC_AVA_LOG(Display, TEXT("FStormSyncAvaSyncProvider::PullFromRemote InRemoteName: %s, InPackageNames: %d"), *InRemoteName, InPackageNames.Num());
+	UE_LOG(LogStormSyncAvaBridge, Display, TEXT("FStormSyncAvaSyncProvider::PullFromRemote InRemoteName: %s, InPackageNames: %d"), *InRemoteName, InPackageNames.Num());
 
 	FText ErrorText;
 	FMessageAddress MessageAddress;
 	// Get address from either client or server user data, depending on local client / server state
 	if (!GetAddressFromUserData(InRemoteName, UE::StormSync::AvaBridgeCommon::StormSyncClientAddressKey, MessageAddress, &ErrorText))
 	{
-		STORM_SYNC_AVA_LOG(Error, TEXT("FStormSyncAvaSyncProvider::PullFromRemote - %s"), *ErrorText.ToString());
+		UE_LOG(LogStormSyncAvaBridge, Error, TEXT("FStormSyncAvaSyncProvider::PullFromRemote - %s"), *ErrorText.ToString());
 		DoneDelegate.ExecuteIfBound(CreateErrorResponse(ErrorText));
 		return;
 	}
@@ -92,7 +92,7 @@ void FStormSyncAvaSyncProvider::PullFromRemote(const FString& InRemoteName, cons
 			return;
 		}
 		
-		STORM_SYNC_AVA_LOG(Display, TEXT("FStormSyncAvaSyncProvider::PullFromRemote - Response: %s"), *Response->ToString());
+		UE_LOG(LogStormSyncAvaBridge, Display, TEXT("FStormSyncAvaSyncProvider::PullFromRemote - Response: %s"), *Response->ToString());
 
 		if (DoneDelegate.IsBound())
 		{
@@ -107,11 +107,11 @@ void FStormSyncAvaSyncProvider::PullFromRemote(const FString& InRemoteName, cons
 
 void FStormSyncAvaSyncProvider::CompareWithRemote(const FString& InRemoteName, const TArray<FName>& InPackageNames, const FOnAvaMediaSyncCompareResponse& DoneDelegate)
 {
-	STORM_SYNC_AVA_LOG(Display, TEXT("FStormSyncAvaSyncProvider::CompareWithRemote InRemoteName: %s, InPackageNames: %d"), *InRemoteName, InPackageNames.Num());
+	UE_LOG(LogStormSyncAvaBridge, Display, TEXT("FStormSyncAvaSyncProvider::CompareWithRemote InRemoteName: %s, InPackageNames: %d"), *InRemoteName, InPackageNames.Num());
 
 	const bool bIsClient = IsPlaybackClient();
 	const bool bIsServer = IsPlaybackServer();
-	STORM_SYNC_AVA_LOG(Display, TEXT("\t bIsClient: %s, bIsServer: %s"), bIsClient ? TEXT("true") : TEXT("false"), bIsServer ? TEXT("true") : TEXT("false"));
+	UE_LOG(LogStormSyncAvaBridge, Display, TEXT("\t bIsClient: %s, bIsServer: %s"), bIsClient ? TEXT("true") : TEXT("false"), bIsServer ? TEXT("true") : TEXT("false"));
 
 	// Sanity check on input params
 	if (InRemoteName.IsEmpty())
@@ -131,7 +131,7 @@ void FStormSyncAvaSyncProvider::CompareWithRemote(const FString& InRemoteName, c
 	// Get address from either client or server user data, depending on local client / server state
 	if (!GetAddressFromUserData(InRemoteName, UE::StormSync::AvaBridgeCommon::StormSyncClientAddressKey, MessageAddress, &ErrorText))
 	{
-		STORM_SYNC_AVA_LOG(Error, TEXT("FStormSyncAvaSyncProvider::CompareWithRemote - %s"), *ErrorText.ToString());
+		UE_LOG(LogStormSyncAvaBridge, Error, TEXT("FStormSyncAvaSyncProvider::CompareWithRemote - %s"), *ErrorText.ToString());
 		DoneDelegate.ExecuteIfBound(CreateErrorResponse(ErrorText));
 		return;
 	}
@@ -147,11 +147,11 @@ FOnAvaMediaSyncPackageModified& FStormSyncAvaSyncProvider::GetOnAvaSyncPackageMo
 
 void FStormSyncAvaSyncProvider::CompareWith(const FMessageAddress& InRemoteAddress, const TArray<FName>& InPackageNames, const FOnAvaMediaSyncCompareResponse& DoneDelegate)
 {
-	STORM_SYNC_AVA_LOG(Display, TEXT("FStormSyncAvaSyncProvider::CompareWith InPackageNames: %d (Remote Address: %s)"), InPackageNames.Num(), *InRemoteAddress.ToString());
+	UE_LOG(LogStormSyncAvaBridge, Display, TEXT("FStormSyncAvaSyncProvider::CompareWith InPackageNames: %d (Remote Address: %s)"), InPackageNames.Num(), *InRemoteAddress.ToString());
 
 	for (const FName& PackageName : InPackageNames)
 	{
-		STORM_SYNC_AVA_LOG(Display, TEXT("\t PackageNames: %s"), *PackageName.ToString());
+		UE_LOG(LogStormSyncAvaBridge, Display, TEXT("\t PackageNames: %s"), *PackageName.ToString());
 	}
 
 	IStormSyncTransportClientModule::Get().RequestPackagesStatus(InRemoteAddress, InPackageNames, FOnStormSyncRequestStatusComplete::CreateLambda([DoneDelegate](const TSharedPtr<FStormSyncTransportStatusResponse>& Response)
@@ -162,7 +162,7 @@ void FStormSyncAvaSyncProvider::CompareWith(const FMessageAddress& InRemoteAddre
 			return;
 		}
 
-		STORM_SYNC_AVA_LOG(Display, TEXT("FStormSyncAvaSyncProvider::CompareWith - Response: %s"), *Response->ToString());
+		UE_LOG(LogStormSyncAvaBridge, Display, TEXT("FStormSyncAvaSyncProvider::CompareWith - Response: %s"), *Response->ToString());
 
 		if (DoneDelegate.IsBound())
 		{
@@ -185,7 +185,7 @@ bool FStormSyncAvaSyncProvider::GetAddressFromUserData(const FString& InRemoteNa
 	const bool bIsClient = IsPlaybackClient();
 	const bool bIsServer = IsPlaybackServer();
 
-	STORM_SYNC_AVA_LOG(
+	UE_LOG(LogStormSyncAvaBridge, 
 		Display,
 		TEXT("FStormSyncAvaSyncProvider::GetAddressFromUserData InRemoteName: %s, bIsClient: %s, bIsServer: %s"),
 		*InRemoteName,
diff --git a/Engine/Plugins/Experimental/StormSyncAvalancheBridge/Source/StormSyncAvaBridge/Private/ModularFeature/StormSyncAvaSyncProvider.h b/Engine/Plugins/Experimental/StormSyncAvalancheBridge/Source/StormSyncAvaBridge/Private/ModularFeature/StormSyncAvaSyncProvider.h
--- a/Engine/Plugins/Experimental/StormSyncAvalancheBridge/Source/StormSyncAvaBridge/Private/ModularFeature/StormSyncAvaSyncProvider.h
+++ b/Engine/Plugins/Experimental/StormSyncAvalancheBridge/Source/StormSyncAvaBridge/Private/ModularFeature/StormSyncAvaSyncProvider.h
@@ -2,8 +2,8 @@
 
 #pragma once
 
-#include "CoreMinimal.h"
 #include "ModularFeature/IAvaMediaSyncProvider.h"
+#include "Templates/SharedPointerFwd.h"
 
 enum class EStormSyncEngineType : uint8;
 struct FMessageAddress;
@@ -54,16 +54,16 @@ private:
 	 * 4. If an instance is both a client and server, the client takes priority and server side is ignored
 	 */
 	static bool GetAddressFromUserData(const FString& InRemoteName, const FString& InUserDataKey, FMessageAddress& OutAddress, FText* OutErrorMessage = nullptr);
-	
+
 	/** Helper to get server user data from playback client */
 	static bool GetAddressFromServerUserData(const FString& InServerName, const FString& InUserDataKey, FMessageAddress& OutAddress, FText* OutErrorMessage = nullptr);
-	
+
 	/** Helper to get client user data from playback server */
 	static bool GetAddressFromClientUserData(const FString& InClientName, const FString& InUserDataKey, FMessageAddress& OutAddress, FText* OutErrorMessage = nullptr);
 
 	/** Returns whether local editor instance has a playback server started */
 	static bool IsPlaybackServer();
-	
+
 	/** Returns whether local editor instance has a playback client started */
 	static bool IsPlaybackClient();
 
@@ -72,7 +72,7 @@ private:
 
 	/** Helper to convert a FStormSyncTransportSyncResponse struct to the modular feature equivalent */
 	static FAvaMediaSyncResponse ConvertSyncResponse(const TSharedPtr<FStormSyncTransportSyncResponse>& InResponse);
-	
+
 	/** Helper to convert a FStormSyncConnectionInfo struct to the modular feature equivalent */
 	static FAvaMediaSyncConnectionInfo ConvertConnectionInfo(const FStormSyncConnectionInfo& InConnectionInfo);
 
diff --git a/Engine/Plugins/Experimental/StormSyncAvalancheBridge/Source/StormSyncAvaBridge/Private/StormSyncAvaBridge.cpp b/Engine/Plugins/Experimental/StormSyncAvalancheBridge/Source/StormSyncAvaBridge/Private/StormSyncAvaBridge.cpp
--- a/Engine/Plugins/Experimental/StormSyncAvalancheBridge/Source/StormSyncAvaBridge/Private/StormSyncAvaBridge.cpp
+++ b/Engine/Plugins/Experimental/StormSyncAvalancheBridge/Source/StormSyncAvaBridge/Private/StormSyncAvaBridge.cpp
@@ -12,9 +12,12 @@
 #include "Playback/IAvaPlaybackClient.h"
 #include "StormSyncAvaBridgeCommon.h"
 #include "StormSyncAvaBridgeLog.h"
+#include "StormSyncAvaBridgeLog.h"
 #include "StormSyncAvaBridgeUtils.h"
 #include "StormSyncCoreDelegates.h"
 
+DEFINE_LOG_CATEGORY(LogStormSyncAvaBridge);
+
 #define LOCTEXT_NAMESPACE "FStormSyncAvaBridgeModule"
 
 void FStormSyncAvaBridgeModule::StartupModule()
@@ -97,29 +100,29 @@ void FStormSyncAvaBridgeModule::UnregisterConsoleCommands()
 void FStormSyncAvaBridgeModule::ExecuteGetUserData(const TArray<FString>& Args)
 {
 	const FString Arguments = FString::Join(Args, TEXT(" "));
-	STORM_SYNC_AVA_LOG(Display, TEXT("FStormSyncAvaBridgeModule::ExecuteGetUserData - Args: %s"), *Arguments)
+	UE_LOG(LogStormSyncAvaBridge, Display, TEXT("FStormSyncAvaBridgeModule::ExecuteGetUserData - Args: %s"), *Arguments);
 
 	if (!Args.IsValidIndex(0) || Args[0].IsEmpty())
 	{
-		STORM_SYNC_AVA_LOG(Error, TEXT("Missing first parameter \"ChannelName\""))
+		UE_LOG(LogStormSyncAvaBridge, Error, TEXT("Missing first parameter \"ChannelName\""));
 		return;
 	}
 
 	if (!Args.IsValidIndex(1) || Args[1].IsEmpty())
 	{
-		STORM_SYNC_AVA_LOG(Error, TEXT("Missing second parameter \"Key\""))
+		UE_LOG(LogStormSyncAvaBridge, Error, TEXT("Missing second parameter \"Key\""));
 		return;
 	}
 
 	if (!IAvaMediaModule::IsModuleLoaded())
 	{
-		STORM_SYNC_AVA_LOG(Error, TEXT("Ava Media Module is not available"))
+		UE_LOG(LogStormSyncAvaBridge, Error, TEXT("Ava Media Module is not available"));
 		return;
 	}
 
 	if (!IAvaMediaModule::Get().IsMediaPlaybackClientStarted())
 	{
-		STORM_SYNC_AVA_LOG(Error, TEXT("Ava Playback Client stopped"))
+		UE_LOG(LogStormSyncAvaBridge, Error, TEXT("Ava Playback Client stopped"));
 		return;
 	}
 
@@ -129,7 +132,7 @@ void FStormSyncAvaBridgeModule::ExecuteGetUserData(const TArray<FString>& Args)
 	TArray<FString> ServerNames = FStormSyncAvaBridgeUtils::GetServerNamesForChannel(ChannelName);
 	if (ServerNames.IsEmpty())
 	{
-		STORM_SYNC_AVA_LOG(Display, TEXT("Ava Broadcast profile for channel %s has no remotes"), *ChannelName)
+		UE_LOG(LogStormSyncAvaBridge, Display, TEXT("Ava Broadcast profile for channel %s has no remotes"), *ChannelName);
 		return;
 	}
 
@@ -137,19 +140,19 @@ void FStormSyncAvaBridgeModule::ExecuteGetUserData(const TArray<FString>& Args)
 	for (const FString& ServerName : ServerNames)
 	{
 		FString ServerValue = PlaybackClient.GetServerUserData(ServerName, Key);
-		STORM_SYNC_AVA_LOG(Display, TEXT("\t Media Playback Client User Data for server %s - %s:%s"), *ServerName, *Key, *ServerValue);
+		UE_LOG(LogStormSyncAvaBridge, Display, TEXT("\t Media Playback Client User Data for server %s - %s:%s"), *ServerName, *Key, *ServerValue);
 	}
 }
 
 void FStormSyncAvaBridgeModule::OnStormSyncServerStarted()
 {
-	STORM_SYNC_AVA_LOG(Verbose, TEXT("FStormSyncAvaBridgeModule::OnStormSyncServerStarted"))
+	UE_LOG(LogStormSyncAvaBridge, Verbose, TEXT("FStormSyncAvaBridgeModule::OnStormSyncServerStarted"));
 	RegisterUserDataForPlaybackServer();
 }
 
 void FStormSyncAvaBridgeModule::OnStormSyncServerStopped()
 {
-	STORM_SYNC_AVA_LOG(Verbose, TEXT("FStormSyncAvaBridgeModule::OnStormSyncServerStopped"))
+	UE_LOG(LogStormSyncAvaBridge, Verbose, TEXT("FStormSyncAvaBridgeModule::OnStormSyncServerStopped"));
 	RegisterUserDataForPlaybackServer();
 }
 
@@ -171,9 +174,9 @@ void FStormSyncAvaBridgeModule::RegisterUserDataForPlaybackServer()
 		PlaybackServer->SetUserData(UE::StormSync::AvaBridgeCommon::StormSyncClientAddressKey, StormSyncClientAddress);
 		PlaybackServer->SetUserData(UE::StormSync::AvaBridgeCommon::StormSyncDiscoveryAddressKey, DiscoveryManagerAddress);
 
-		STORM_SYNC_AVA_LOG(Display, TEXT("FStormSyncAvaBridgeModule::RegisterUserDataForPlaybackServer StormSyncServerAddress: %s"), *StormSyncServerAddress)
-		STORM_SYNC_AVA_LOG(Display, TEXT("FStormSyncAvaBridgeModule::RegisterUserDataForPlaybackServer StormSyncClientAddress: %s"), *StormSyncClientAddress)
-		STORM_SYNC_AVA_LOG(Display, TEXT("FStormSyncAvaBridgeModule::RegisterUserDataForPlaybackServer DiscoveryManagerAddress: %s"), *DiscoveryManagerAddress)
+		UE_LOG(LogStormSyncAvaBridge, Display, TEXT("FStormSyncAvaBridgeModule::RegisterUserDataForPlaybackServer StormSyncServerAddress: %s"), *StormSyncServerAddress);
+		UE_LOG(LogStormSyncAvaBridge, Display, TEXT("FStormSyncAvaBridgeModule::RegisterUserDataForPlaybackServer StormSyncClientAddress: %s"), *StormSyncClientAddress);
+		UE_LOG(LogStormSyncAvaBridge, Display, TEXT("FStormSyncAvaBridgeModule::RegisterUserDataForPlaybackServer DiscoveryManagerAddress: %s"), *DiscoveryManagerAddress);
 	}
 }
 
@@ -197,9 +200,9 @@ void FStormSyncAvaBridgeModule::RegisterUserDataForPlaybackClient()
 		PlaybackClient.SetUserData(UE::StormSync::AvaBridgeCommon::StormSyncClientAddressKey, StormSyncClientAddress);
 		PlaybackClient.SetUserData(UE::StormSync::AvaBridgeCommon::StormSyncDiscoveryAddressKey, DiscoveryManagerAddress);
 
-		STORM_SYNC_AVA_LOG(Display, TEXT("FStormSyncAvaBridgeModule::RegisterUserDataForPlaybackClient StormSyncServerAddress: %s"), *StormSyncServerAddress)
-		STORM_SYNC_AVA_LOG(Display, TEXT("FStormSyncAvaBridgeModule::RegisterUserDataForPlaybackClient StormSyncClientAddress: %s"), *StormSyncClientAddress)
-		STORM_SYNC_AVA_LOG(Display, TEXT("FStormSyncAvaBridgeModule::RegisterUserDataForPlaybackClient DiscoveryManagerAddress: %s"), *DiscoveryManagerAddress)
+		UE_LOG(LogStormSyncAvaBridge, Display, TEXT("FStormSyncAvaBridgeModule::RegisterUserDataForPlaybackClient StormSyncServerAddress: %s"), *StormSyncServerAddress);
+		UE_LOG(LogStormSyncAvaBridge, Display, TEXT("FStormSyncAvaBridgeModule::RegisterUserDataForPlaybackClient StormSyncClientAddress: %s"), *StormSyncClientAddress);
+		UE_LOG(LogStormSyncAvaBridge, Display, TEXT("FStormSyncAvaBridgeModule::RegisterUserDataForPlaybackClient DiscoveryManagerAddress: %s"), *DiscoveryManagerAddress);
 	}
 }
 
@@ -207,19 +210,19 @@ bool FStormSyncAvaBridgeModule::ValidateModulesAreAvailable()
 {
 	if (!IAvaMediaModule::IsModuleLoaded())
 	{
-		STORM_SYNC_AVA_LOG(Warning, TEXT("FStormSyncAvaBridgeModule::ValidateModulesAreAvailable - Failed to set user data cause Ava Media Module is not loaded"))
+		UE_LOG(LogStormSyncAvaBridge, Warning, TEXT("FStormSyncAvaBridgeModule::ValidateModulesAreAvailable - Failed to set user data cause Ava Media Module is not loaded"));
 		return false;
 	}
 
 	if (!IStormSyncTransportServerModule::IsAvailable())
 	{
-		STORM_SYNC_AVA_LOG(Warning, TEXT("FStormSyncAvaBridgeModule::ValidateModulesAreAvailable - Failed to set user data cause Storm Sync Server Module is not loaded"))
+		UE_LOG(LogStormSyncAvaBridge, Warning, TEXT("FStormSyncAvaBridgeModule::ValidateModulesAreAvailable - Failed to set user data cause Storm Sync Server Module is not loaded"));
 		return false;
 	}
 
 	if (!IStormSyncTransportClientModule::IsAvailable())
 	{
-		STORM_SYNC_AVA_LOG(Warning, TEXT("FStormSyncAvaBridgeModule::ValidateModulesAreAvailable - Failed to set user data cause Storm Sync Server Client is not loaded"))
+		UE_LOG(LogStormSyncAvaBridge, Warning, TEXT("FStormSyncAvaBridgeModule::ValidateModulesAreAvailable - Failed to set user data cause Storm Sync Server Client is not loaded"));
 		return false;
 	}
 
diff --git a/Engine/Plugins/Experimental/StormSyncAvalancheBridge/Source/StormSyncAvaBridge/Private/StormSyncAvaBridge.h b/Engine/Plugins/Experimental/StormSyncAvalancheBridge/Source/StormSyncAvaBridge/Private/StormSyncAvaBridge.h
new file mode 100644
--- /dev/null
+++ b/Engine/Plugins/Experimental/StormSyncAvalancheBridge/Source/StormSyncAvaBridge/Private/StormSyncAvaBridge.h
@@ -0,0 +1,69 @@
+// Copyright Epic Games, Inc. All Rights Reserved.
+
+#pragma once
+
+#include "Containers/Array.h"
+#include "Containers/StringFwd.h"
+#include "Modules/ModuleInterface.h"
+#include "Templates/UniquePtr.h"
+
+class FStormSyncAvaSyncProvider;
+class IConsoleObject;
+
+/**
+ * Module for Storm Sync Ava Bridge runtime module
+ *
+ * This class handles the update of User Data in Ava Media Playback Server when module is loaded,
+ * and Storm Sync server is started or stopped.
+ *
+ * This is then replicated and made available on remote clients so that a Playlist asset can associate
+ * a message bus address id (for storm sync push / sync) to a remote on a channel.
+ */
+class FStormSyncAvaBridgeModule : public IModuleInterface
+{
+	//~ Begin IModuleInterface interface
+	virtual void StartupModule() override;
+	virtual void ShutdownModule() override;
+	//~ End IModuleInterface interface
+
+	/** References of registered console commands via IConsoleManager */
+	TArray<IConsoleObject*> ConsoleCommands;
+
+	/** Ava sync provider feature instance */
+	TUniquePtr<FStormSyncAvaSyncProvider> SyncProvider;
+
+	/** Updates playback user data when engine is fully loaded to setup storm sync message address ids */
+	void OnPostEngineInit();
+
+	/** Called from StartupModule and sets up console commands for the plugin via IConsoleManager */
+	void RegisterConsoleCommands();
+
+	/** Called from ShutdownModule and clears out previously registered console commands */
+	void UnregisterConsoleCommands();
+
+	/** Console command handler to return associated user data in ava playback server */
+	static void ExecuteGetUserData(const TArray<FString>& Args);
+
+	/** Updates user data when storm sync server is started */
+	void OnStormSyncServerStarted();
+	
+	/** Updates user data when storm sync server is stopped */
+	void OnStormSyncServerStopped();
+
+	/**
+	 * Actual helper to set the user data for a running playback server.
+	 * 
+	 * Sets up Storm Sync Server / Client and Discovery Manager address ids
+	 */
+	static void RegisterUserDataForPlaybackServer();
+	
+	/**
+	 * Actual helper to set the user data for a running playback client.
+	 * 
+	 * Sets up Storm Sync Server / Client and Discovery Manager address ids
+	 */
+	static void RegisterUserDataForPlaybackClient();
+
+	/** Helper ensuring all required modules are loaded available */
+	static bool ValidateModulesAreAvailable();
+};
diff --git a/Engine/Plugins/Experimental/StormSyncAvalancheBridge/Source/StormSyncAvaBridge/Private/StormSyncAvaBridgeLog.cpp b/Engine/Plugins/Experimental/StormSyncAvalancheBridge/Source/StormSyncAvaBridge/Private/StormSyncAvaBridgeLog.cpp
deleted file mode 100644
diff --git a/Engine/Plugins/Experimental/StormSyncAvalancheBridge/Source/StormSyncAvaBridge/Private/StormSyncAvaBridgeLog.h b/Engine/Plugins/Experimental/StormSyncAvalancheBridge/Source/StormSyncAvaBridge/Private/StormSyncAvaBridgeLog.h
new file mode 100644
--- /dev/null
+++ b/Engine/Plugins/Experimental/StormSyncAvalancheBridge/Source/StormSyncAvaBridge/Private/StormSyncAvaBridgeLog.h
@@ -0,0 +1,7 @@
+// Copyright Epic Games, Inc. All Rights Reserved.
+
+#pragma once
+
+#include "Logging/LogMacros.h"
+
+DECLARE_LOG_CATEGORY_EXTERN(LogStormSyncAvaBridge, Display, All);
diff --git a/Engine/Plugins/Experimental/StormSyncAvalancheBridge/Source/StormSyncAvaBridge/Public/StormSyncAvaBridge.h b/Engine/Plugins/Experimental/StormSyncAvalancheBridge/Source/StormSyncAvaBridge/Public/StormSyncAvaBridge.h
deleted file mode 100644
diff --git a/Engine/Plugins/Experimental/StormSyncAvalancheBridge/Source/StormSyncAvaBridge/Public/StormSyncAvaBridgeCommon.h b/Engine/Plugins/Experimental/StormSyncAvalancheBridge/Source/StormSyncAvaBridge/Public/StormSyncAvaBridgeCommon.h
--- a/Engine/Plugins/Experimental/StormSyncAvalancheBridge/Source/StormSyncAvaBridge/Public/StormSyncAvaBridgeCommon.h
+++ b/Engine/Plugins/Experimental/StormSyncAvalancheBridge/Source/StormSyncAvaBridge/Public/StormSyncAvaBridgeCommon.h
@@ -2,16 +2,16 @@
 
 #pragma once
 
-#include "CoreMinimal.h"
+#include "HAL/Platform.h"
 
 namespace UE::StormSync::AvaBridgeCommon
 {
 	/** The user data key to store storm sync server address id */
 	static constexpr const TCHAR* StormSyncServerAddressKey = TEXT("StormSyncServerAddress");
-	
+
 	/** The user data key to store storm sync discovery manager address id */
 	static constexpr const TCHAR* StormSyncClientAddressKey = TEXT("StormSyncClientAddress");
 
 	/** The user data key to store storm sync discovery manager address id */
 	static constexpr const TCHAR* StormSyncDiscoveryAddressKey = TEXT("StormSyncDiscoveryAddress");
-}
\ No newline at end of file
+}
diff --git a/Engine/Plugins/Experimental/StormSyncAvalancheBridge/Source/StormSyncAvaBridge/Public/StormSyncAvaBridgeLog.h b/Engine/Plugins/Experimental/StormSyncAvalancheBridge/Source/StormSyncAvaBridge/Public/StormSyncAvaBridgeLog.h
deleted file mode 100644
diff --git a/Engine/Plugins/Experimental/StormSyncAvalancheBridge/Source/StormSyncAvaBridge/Public/StormSyncAvaBridgeUtils.h b/Engine/Plugins/Experimental/StormSyncAvalancheBridge/Source/StormSyncAvaBridge/Public/StormSyncAvaBridgeUtils.h
--- a/Engine/Plugins/Experimental/StormSyncAvalancheBridge/Source/StormSyncAvaBridge/Public/StormSyncAvaBridgeUtils.h
+++ b/Engine/Plugins/Experimental/StormSyncAvalancheBridge/Source/StormSyncAvaBridge/Public/StormSyncAvaBridgeUtils.h
@@ -2,7 +2,8 @@
 
 #pragma once
 
-#include "CoreMinimal.h"
+#include "Containers/ContainersFwd.h"
+#include "Containers/StringFwd.h"
 
 class STORMSYNCAVABRIDGE_API FStormSyncAvaBridgeUtils
 {
diff --git a/Engine/Plugins/Experimental/StormSyncAvalancheBridge/Source/StormSyncAvaBridgeEditor/Private/StormSyncAvaBridgeEditor.cpp b/Engine/Plugins/Experimental/StormSyncAvalancheBridge/Source/StormSyncAvaBridgeEditor/Private/StormSyncAvaBridgeEditor.cpp
--- a/Engine/Plugins/Experimental/StormSyncAvalancheBridge/Source/StormSyncAvaBridgeEditor/Private/StormSyncAvaBridgeEditor.cpp
+++ b/Engine/Plugins/Experimental/StormSyncAvalancheBridge/Source/StormSyncAvaBridgeEditor/Private/StormSyncAvaBridgeEditor.cpp
@@ -1,20 +1,20 @@
 // Copyright Epic Games, Inc. All Rights Reserved.
 
 #include "StormSyncAvaBridgeEditor.h"
+#include "Modules/ModuleManager.h"
+#include "StormSyncAvaBridgeEditorLog.h"
+#include "StormSyncAvaRundownExtender.h"
 
-#include "StormSyncAvaPlaylistExtender.h"
+DEFINE_LOG_CATEGORY(LogStormSyncAvaBridgeEditor);
 
 void FStormSyncAvaBridgeEditorModule::StartupModule()
 {
-	PlaylistExtender = MakeShared<FStormSyncAvaPlaylistExtender>();
+	RundownExtender = MakeShared<FStormSyncAvaRundownExtender>();
 }
 
 void FStormSyncAvaBridgeEditorModule::ShutdownModule()
 {
-	if (PlaylistExtender.IsValid())
-	{
-		PlaylistExtender.Reset();
-	}
+	RundownExtender.Reset();
 }
 
 IMPLEMENT_MODULE(FStormSyncAvaBridgeEditorModule, StormSyncAvaBridgeEditor)
diff --git a/Engine/Plugins/Experimental/StormSyncAvalancheBridge/Source/StormSyncAvaBridgeEditor/Private/StormSyncAvaBridgeEditor.h b/Engine/Plugins/Experimental/StormSyncAvalancheBridge/Source/StormSyncAvaBridgeEditor/Private/StormSyncAvaBridgeEditor.h
new file mode 100644
--- /dev/null
+++ b/Engine/Plugins/Experimental/StormSyncAvalancheBridge/Source/StormSyncAvaBridgeEditor/Private/StormSyncAvaBridgeEditor.h
@@ -0,0 +1,19 @@
+// Copyright Epic Games, Inc. All Rights Reserved.
+
+#pragma once
+
+#include "Modules/ModuleInterface.h"
+#include "Templates/SharedPointer.h"
+
+class FStormSyncAvaRundownExtender;
+
+class FStormSyncAvaBridgeEditorModule : public IModuleInterface
+{
+	//~ Begin IModuleInterface interface
+	virtual void StartupModule() override;
+	virtual void ShutdownModule() override;
+	//~ End IModuleInterface interface
+
+	/** UI extender for Storm Sync view and commands. */
+	TSharedPtr<FStormSyncAvaRundownExtender> RundownExtender;
+};
diff --git a/Engine/Plugins/Experimental/StormSyncAvalancheBridge/Source/StormSyncAvaBridgeEditor/Private/StormSyncAvaBridgeEditorLog.cpp b/Engine/Plugins/Experimental/StormSyncAvalancheBridge/Source/StormSyncAvaBridgeEditor/Private/StormSyncAvaBridgeEditorLog.cpp
deleted file mode 100644
diff --git a/Engine/Plugins/Experimental/StormSyncAvalancheBridge/Source/StormSyncAvaBridgeEditor/Private/StormSyncAvaBridgeEditorLog.h b/Engine/Plugins/Experimental/StormSyncAvalancheBridge/Source/StormSyncAvaBridgeEditor/Private/StormSyncAvaBridgeEditorLog.h
new file mode 100644
--- /dev/null
+++ b/Engine/Plugins/Experimental/StormSyncAvalancheBridge/Source/StormSyncAvaBridgeEditor/Private/StormSyncAvaBridgeEditorLog.h
@@ -0,0 +1,7 @@
+// Copyright Epic Games, Inc. All Rights Reserved.
+
+#pragma once
+
+#include "Logging/LogMacros.h"
+
+DECLARE_LOG_CATEGORY_EXTERN(LogStormSyncAvaBridgeEditor, Display, All);
diff --git a/Engine/Plugins/Experimental/StormSyncAvalancheBridge/Source/StormSyncAvaBridgeEditor/Private/StormSyncAvaPlaylistExtender.cpp b/Engine/Plugins/Experimental/StormSyncAvalancheBridge/Source/StormSyncAvaBridgeEditor/Private/StormSyncAvaPlaylistExtender.cpp
deleted file mode 100644
diff --git a/Engine/Plugins/Experimental/StormSyncAvalancheBridge/Source/StormSyncAvaBridgeEditor/Private/StormSyncAvaPlaylistExtender.h b/Engine/Plugins/Experimental/StormSyncAvalancheBridge/Source/StormSyncAvaBridgeEditor/Private/StormSyncAvaPlaylistExtender.h
deleted file mode 100644
diff --git a/Engine/Plugins/Experimental/StormSyncAvalancheBridge/Source/StormSyncAvaBridgeEditor/Private/StormSyncAvaRundownExtender.cpp b/Engine/Plugins/Experimental/StormSyncAvalancheBridge/Source/StormSyncAvaBridgeEditor/Private/StormSyncAvaRundownExtender.cpp
new file mode 100644
--- /dev/null
+++ b/Engine/Plugins/Experimental/StormSyncAvalancheBridge/Source/StormSyncAvaBridgeEditor/Private/StormSyncAvaRundownExtender.cpp
@@ -0,0 +1,589 @@
+// Copyright Epic Games, Inc. All Rights Reserved.
+
+#include "StormSyncAvaRundownExtender.h"
+#include "Framework/MultiBox/MultiBoxBuilder.h"
+#include "IAvaMediaEditorModule.h"
+#include "IAvaMediaModule.h"
+#include "IStormSyncTransportClientModule.h"
+#include "Playback/IAvaPlaybackClient.h"
+#include "Rundown/AvaRundown.h"
+#include "Rundown/AvaRundownEditor.h"
+#include "StormSyncAvaBridgeCommon.h"
+#include "StormSyncAvaBridgeEditorLog.h"
+#include "StormSyncAvaBridgeUtils.h"
+#include "StormSyncEditor.h"
+#include "Subsystems/StormSyncNotificationSubsystem.h"
+#include "Toolkits/ToolkitManager.h"
+
+#define LOCTEXT_NAMESPACE "FStormSyncAvaRundownExtender"
+
+FStormSyncAvaRundownExtender::FStormSyncAvaRundownExtender()
+{
+	FCoreDelegates::OnPostEngineInit.AddRaw(this, &FStormSyncAvaRundownExtender::OnPostEngineInit);
+}
+
+FStormSyncAvaRundownExtender::~FStormSyncAvaRundownExtender()
+{
+	FCoreDelegates::OnPostEngineInit.RemoveAll(this);
+
+	if (IAvaMediaEditorModule::IsLoaded())
+	{
+		UnregisterExtension(IAvaMediaEditorModule::Get().GetRundownMenuExtensibilityManager(), MenuExtenderHandle);
+		UnregisterExtension(IAvaMediaEditorModule::Get().GetRundownToolBarExtensibilityManager(), ToolbarExtenderHandle);
+	}
+}
+
+void FStormSyncAvaRundownExtender::OnPostEngineInit()
+{
+	if (IAvaMediaEditorModule::IsLoaded())
+	{
+		RegisterMenuExtensions();
+	}
+}
+
+void FStormSyncAvaRundownExtender::RegisterMenuExtensions()
+{
+	// Context menu extension
+	MenuExtenderHandle = RegisterExtension(
+		IAvaMediaEditorModule::Get().GetRundownMenuExtensibilityManager(),
+		FAssetEditorExtender::CreateSP(this, &FStormSyncAvaRundownExtender::AddMenuExtender)
+	);
+
+	// Toolbar extension
+	ToolbarExtenderHandle = RegisterExtension(
+		IAvaMediaEditorModule::Get().GetRundownToolBarExtensibilityManager(),
+		FAssetEditorExtender::CreateSP(this, &FStormSyncAvaRundownExtender::AddToolbarExtender)
+	);
+}
+
+FDelegateHandle FStormSyncAvaRundownExtender::RegisterExtension(const TSharedPtr<FExtensibilityManager> InExtensibilityManager, const FAssetEditorExtender& InExtenderDelegate)
+{
+	if (!InExtensibilityManager.IsValid())
+	{
+		const FDelegateHandle DelegateHandle;
+		return DelegateHandle;
+	}
+
+	const int32 ExtenderIndex = InExtensibilityManager->GetExtenderDelegates().Add(InExtenderDelegate);
+	return InExtensibilityManager->GetExtenderDelegates()[ExtenderIndex].GetHandle();
+}
+
+void FStormSyncAvaRundownExtender::UnregisterExtension(const TSharedPtr<FExtensibilityManager> InExtensibilityManager, const FDelegateHandle& InHandleToRemove)
+{
+	if (!InExtensibilityManager.IsValid())
+	{
+		return;
+	}
+
+	InExtensibilityManager->GetExtenderDelegates().RemoveAll([Handle = InHandleToRemove](const FAssetEditorExtender& Extender)
+	{
+		return Handle == Extender.GetHandle();
+	});
+}
+
+TSharedRef<FExtender> FStormSyncAvaRundownExtender::AddMenuExtender(const TSharedRef<FUICommandList> InCommandList, const TArray<UObject*> ContextSensitiveObjects)
+{
+	UE_LOG(LogStormSyncAvaBridgeEditor, Display, TEXT("FStormSyncAvaRundownExtender::AddMenuExtender - Adding in menu extensions"));
+
+	TSharedRef<FExtender> Extender(new FExtender());
+
+	const UAvaRundown* Rundown = ContextSensitiveObjects.IsValidIndex(0) ? Cast<UAvaRundown>(ContextSensitiveObjects[0]) : nullptr;
+	if (!Rundown)
+	{
+		return Extender;
+	}
+
+	const TSharedPtr<IToolkit> AssetEditor = FToolkitManager::Get().FindEditorForAsset(Rundown);
+	if (!AssetEditor.IsValid())
+	{
+		return Extender;
+	}
+
+	const TSharedPtr<FAvaRundownEditor> RundownEditor = StaticCastSharedPtr<FAvaRundownEditor>(AssetEditor);
+	if (!RundownEditor.IsValid())
+	{
+		return Extender;
+	}
+
+	// Template panel extension
+	Extender->AddMenuExtension(
+		AvalancheExtensionHook,
+		EExtensionHook::After,
+		InCommandList,
+		// Convert to weak ptr to prevent ownership to the rundown editor and potentially increasing its lifetime
+		FMenuExtensionDelegate::CreateSP(this, &FStormSyncAvaRundownExtender::CreateTemplateContextMenu, Rundown, TWeakPtr<FAvaRundownEditor>(RundownEditor))
+	);
+
+	return Extender;
+}
+
+void FStormSyncAvaRundownExtender::CreateTemplateContextMenu(FMenuBuilder& MenuBuilder, const UAvaRundown* InRundown, TWeakPtr<FAvaRundownEditor> InRundownEditor)
+{
+	check(InRundown);
+
+	const TSharedPtr<FAvaRundownEditor> RundownEditor = InRundownEditor.Pin();
+	if (!RundownEditor.IsValid())
+	{
+		UE_LOG(LogStormSyncAvaBridgeEditor, Error, TEXT("FStormSyncAvaRundownExtender::CreateTemplateContextMenu - Invalid shared ptr from weak ptr delegate param"));
+		return;
+	}
+
+	if (!FModuleManager::Get().IsModuleLoaded(TEXT("StormSyncEditor")))
+	{
+		UE_LOG(LogStormSyncAvaBridgeEditor, Error, TEXT("FStormSyncAvaRundownExtender::CreateTemplateContextMenu - StormSyncEditor module is not loaded. Is StormSync plugin enabled ?"));
+		return;
+	}
+
+	FText DisabledTooltipReason;
+	bool bIsValidSelection = false;
+	TArray<FName> SelectedPackageNames;
+	FCanExecuteAction DefaultCanExecuteAction;
+	GetContextMenuSelectionInfos(InRundown, InRundownEditor, bIsValidSelection, DisabledTooltipReason, SelectedPackageNames, DefaultCanExecuteAction);
+
+	// For later use with action handler
+	const FStormSyncEditorModule& StormSyncEditor = FStormSyncEditorModule::Get();
+
+	MenuBuilder.BeginSection("StormSyncOperations_Template", LOCTEXT("StormSyncOperationsHeader", "Synchronize Actions"));
+
+	// Initialize action
+	{
+		const FText TooltipText = bIsValidSelection ?
+			FText::Format(LOCTEXT("Initialize_Tooltip", "Sync asset over remote node.{0}"), DisabledTooltipReason) :
+			FText::Format(LOCTEXT("Initialize_Tooltip_Invalid", "Please ensure the rundown page is using a valid Avalanche Blueprint.{0}"), DisabledTooltipReason);
+
+		const TArray<FAvaRundownPage> SelectedTemplatePages = GetSelectedPages(InRundown, InRundownEditor);
+
+		MenuBuilder.AddMenuEntry(
+			LOCTEXT("Initialize_Label", "Initialize"),
+			TooltipText,
+			FSlateIcon(FAppStyle::GetAppStyleSetName(), "LevelEditor.ExportAll"),
+			FUIAction(
+				FExecuteAction::CreateSP(this, &FStormSyncAvaRundownExtender::HandleInitializeAction, InRundown, SelectedTemplatePages),
+				DefaultCanExecuteAction
+			)
+		);
+	}
+
+	constexpr bool bOpenSubMenuOnClick = false;
+
+	// Push action
+	{
+		FText LabelText;
+		FText TooltipText;
+
+		if (!bIsValidSelection)
+		{
+			LabelText = LOCTEXT("PushAssetsMenuEntryInvalid", "Cannot push. Page has no valid asset.");
+			TooltipText = LOCTEXT("PushAssetsMenuEntryTooltipInvalid", "Please ensure the rundown pages are using a valid Avalanche Blueprint.");
+		}
+		else
+		{
+			LabelText = FText::Format(
+				LOCTEXT("PushAssetsAssetName", "Push {0} selected asset{1} to"),
+				FText::AsNumber(SelectedPackageNames.Num()),
+				FText::FromString(SelectedPackageNames.Num() == 1 ? TEXT("") : TEXT("s"))
+			);
+
+			TooltipText = FText::Format(
+				LOCTEXT("PushAssetsMenuEntryTooltip", "Push {0} selected asset{1} (and inner dependencies) to specific remote.\n\nTransfer will only proceed if changes are detected.{2}"),
+				FText::AsNumber(SelectedPackageNames.Num()),
+				FText::FromString(SelectedPackageNames.Num() == 1 ? TEXT("") : TEXT("s")),
+				DisabledTooltipReason
+			);
+		}
+
+		constexpr bool bIsPushing = true;
+		MenuBuilder.AddSubMenu(
+			LabelText,
+			TooltipText,
+			FNewMenuDelegate::CreateRaw(
+				&StormSyncEditor,
+				&FStormSyncEditorModule::BuildPushAssetsMenuSection,
+				SelectedPackageNames,
+				bIsPushing
+			),
+			FUIAction(FExecuteAction(), DefaultCanExecuteAction),
+			NAME_None,
+			EUserInterfaceActionType::Button,
+			bOpenSubMenuOnClick,
+			FSlateIcon(FAppStyle::GetAppStyleSetName(), "LevelEditor.ExportAll")
+		);
+	}
+
+	// Compare action
+	{
+		FText LabelText;
+		FText TooltipText;
+
+		if (!bIsValidSelection)
+		{
+			LabelText = LOCTEXT("CompareAssetsMenuEntryInvalid", "Cannot compare. Page has no valid asset.");
+			TooltipText = LOCTEXT("CompareAssetsMenuEntryTooltipInvalid", "Please ensure the rundown pages are using a valid Avalanche Blueprint.");
+		}
+		else
+		{
+			LabelText = LOCTEXT("CompareAssetsMenuEntry", "Compare Asset(s) with");
+			TooltipText = FText::Format(
+				LOCTEXT("CompareAssetsRemoteMenuEntry", "Compare asset(s) with a specific remote and see if files (and inner dependencies) are either missing or in mismatched state.{0}"),
+				DisabledTooltipReason
+			);
+		}
+
+		MenuBuilder.AddSubMenu(
+			LabelText,
+			TooltipText,
+			FNewMenuDelegate::CreateRaw(
+				&StormSyncEditor,
+				&FStormSyncEditorModule::BuildCompareWithMenuSection,
+				SelectedPackageNames
+			),
+			FUIAction(FExecuteAction(), DefaultCanExecuteAction),
+			NAME_None,
+			EUserInterfaceActionType::Button,
+			bOpenSubMenuOnClick,
+			FSlateIcon(FAppStyle::GetAppStyleSetName(), "EditorViewport.Visualizers")
+		);
+	}
+	MenuBuilder.EndSection();
+}
+
+void FStormSyncAvaRundownExtender::HandleInitializeAction(const UAvaRundown* InRundown, TArray<FAvaRundownPage> InSelectedTemplatePages)
+{
+	// Build a list of package names to push grouped by channel name (and remote)
+	// Key is the remote address id, value is the list of package names to synchronize
+	TMap<FString, TArray<FName>> PackageNamesPerChannel;
+
+	for (const FAvaRundownPage& SelectedPage : InSelectedTemplatePages)
+	{
+		const FString AssetName = SelectedPage.GetAssetPath(InRundown).GetLongPackageName();
+		TArray<FString> ChannelNames = GetChannelNamesForTemplatePage(InRundown, SelectedPage);
+
+		// From the list of channel names that match this asset to sync, build the list of server names
+		// Note: there may be more than one server per channel (channel support multiple outputs).
+
+		UE_LOG(LogStormSyncAvaBridgeEditor, Display, TEXT("ChannelNames for asset \"%s\" are: %s"), *AssetName, *FString::Join(ChannelNames, TEXT(", ")));
+
+		TArray<FString> RemoteServerNames;
+		for (FString ChannelName : ChannelNames)
+		{
+			const TArray<FString> ServerNames = FStormSyncAvaBridgeUtils::GetServerNamesForChannel(ChannelName);
+			if (ServerNames.IsEmpty())
+			{
+				UE_LOG(LogStormSyncAvaBridgeEditor, Warning, TEXT("FStormSyncAvaRundownExtender::HandleInitializeActions - Unable to determine playback servers for channel \"%s\""), *ChannelName);
+				continue;
+			}
+
+			RemoteServerNames.Append(ServerNames);
+		}
+
+		if (IAvaMediaModule::Get().IsMediaPlaybackClientStarted())
+		{
+			const IAvaPlaybackClient& PlaybackClient = IAvaMediaModule::Get().GetMediaPlaybackClient();
+			for (const FString& ServerName : RemoteServerNames)
+			{
+				// Get address id for storm sync client on playback host
+				FString ClientAddress = PlaybackClient.GetServerUserData(ServerName, UE::StormSync::AvaBridgeCommon::StormSyncClientAddressKey);
+				UE_LOG(LogStormSyncAvaBridgeEditor, Display, TEXT("ClientAddress on playback server %s is %s"), *ServerName, *ClientAddress);
+
+				if (ClientAddress.IsEmpty())
+				{
+					UE_LOG(LogStormSyncAvaBridgeEditor, Warning, TEXT("Storm Sync Server Adress id for playback server %s is empty"), *ServerName);
+					continue;
+				}
+
+				TArray<FName>& PackageNames = PackageNamesPerChannel.FindOrAdd(ClientAddress);
+				PackageNames.Add(FName(*AssetName));
+			}
+		}
+	}
+
+	for (const TPair<FString, TArray<FName>>& NamesPerChannel : PackageNamesPerChannel)
+	{
+		FString AddressId = NamesPerChannel.Key;
+		TArray<FName> PackageNames = NamesPerChannel.Value;
+
+		UE_LOG(LogStormSyncAvaBridgeEditor, Display, TEXT("Names per channel - Channel: %s"), *AddressId);
+
+		for (const FName& PackageName : PackageNames)
+		{
+			UE_LOG(LogStormSyncAvaBridgeEditor, Display, TEXT("\t PackageName: %s"), *PackageName.ToString());
+		}
+
+		PushPackagesToRemote(AddressId, PackageNames);
+	}
+}
+
+TSharedRef<FExtender> FStormSyncAvaRundownExtender::AddToolbarExtender(const TSharedRef<FUICommandList> InCommandList, const TArray<UObject*> ContextSensitiveObjects)
+{
+	UE_LOG(LogStormSyncAvaBridgeEditor, Display, TEXT("FStormSyncAvaRundownExtender::AddToolbarExtender - Adding in toolbar extensions"));
+
+	TSharedRef<FExtender> Extender(new FExtender());
+
+	const UAvaRundown* Rundown = ContextSensitiveObjects.IsValidIndex(0) ? Cast<UAvaRundown>(ContextSensitiveObjects[0]) : nullptr;
+	if (!Rundown)
+	{
+		return Extender;
+	}
+
+	const TSharedPtr<IToolkit> AssetEditor = FToolkitManager::Get().FindEditorForAsset(Rundown);
+	if (!AssetEditor.IsValid())
+	{
+		return Extender;
+	}
+
+	const TSharedPtr<FAvaRundownEditor> RundownEditor = StaticCastSharedPtr<FAvaRundownEditor>(AssetEditor);
+	if (!RundownEditor.IsValid())
+	{
+		return Extender;
+	}
+
+	Extender->AddToolBarExtension(
+		"Pages",
+		EExtensionHook::After,
+		InCommandList,
+		FToolBarExtensionDelegate::CreateSP(this, &FStormSyncAvaRundownExtender::FillToolbar, TWeakPtr<FAvaRundownEditor>(RundownEditor))
+	);
+
+	UE_LOG(LogStormSyncAvaBridgeEditor, Display, TEXT("FStormSyncAvaRundownExtender::AddToolbarExtender - Rundown: %s"), *GetNameSafe(Rundown));
+
+	return Extender;
+}
+
+void FStormSyncAvaRundownExtender::FillToolbar(FToolBarBuilder& ToolbarBuilder, TWeakPtr<FAvaRundownEditor> InRundownEditor)
+{
+	ToolbarBuilder.BeginSection(TEXT("StormSync"));
+	{
+		ToolbarBuilder.AddComboButton(
+			FUIAction(),
+			FOnGetContent::CreateSP(this, &FStormSyncAvaRundownExtender::GenerateToolbarMenu, InRundownEditor),
+			LOCTEXT("ToolbarLabel", "Synchronize Actions"),
+			LOCTEXT("ToolbarToolTip", "Synchronize the rundown assets over the network"),
+			FSlateIcon(FAppStyle::GetAppStyleSetName(), "LevelEditor.Recompile"),
+			false
+		);
+	}
+	ToolbarBuilder.EndSection();
+}
+
+TSharedRef<SWidget> FStormSyncAvaRundownExtender::GenerateToolbarMenu(TWeakPtr<FAvaRundownEditor> InRundownEditor)
+{
+	TArray<FName> PackageNames;
+
+	const TSharedPtr<FAvaRundownEditor> RundownEditor = InRundownEditor.Pin();
+
+	if (RundownEditor.IsValid() && RundownEditor->IsRundownValid())
+	{
+		const UAvaRundown* Rundown = RundownEditor->GetRundown(); 
+
+		// Gather the list of all package names from Avalanche Blueprints in this Rundown pages
+		const FAvaRundownPageCollection& AvalanchePageCollection = Rundown->GetTemplatePages();
+		for (const FAvaRundownPage& AvalanchePage : AvalanchePageCollection.Pages)
+		{
+			if (FString PackageName = AvalanchePage.GetAssetPath(Rundown).GetLongPackageName(); !PackageName.IsEmpty())
+			{
+				PackageNames.AddUnique(FName(*PackageName));
+			}
+		}
+	}
+
+	FMenuBuilder MenuBuilder(true, nullptr);
+
+	MenuBuilder.BeginSection(TEXT("StormSyncActions"));
+
+	if (FModuleManager::Get().IsModuleLoaded(TEXT("StormSyncEditor")))
+	{
+		FStormSyncEditorModule& StormSyncEditor = FStormSyncEditorModule::Get();
+
+		const int32 PackagesCount = PackageNames.Num();
+
+		TArray<FString> AssetList;
+		for (const FName PackageName : PackageNames)
+		{
+			AssetList.Add(FString::Printf(TEXT("- %s"), *PackageName.ToString()));
+		}
+
+		FText LabelText = FText::Format(LOCTEXT("PushAssetsToolbarMenuEntry", "Push {0} asset(s) to"), FText::AsNumber(PackagesCount));
+		FText TooltipText = FText::Format(LOCTEXT(
+			"PushAssetsToolbarMenuEntryTooltip",
+			"Push {0} asset(s) (and inner dependencies) to specific remote.\n\nTransfer will only proceed if changes are detected.\n\n{1}"
+		), FText::AsNumber(PackagesCount), FText::FromString(FString::Join(AssetList, LINE_TERMINATOR)));
+
+		const bool bIsPushEnabled = !PackageNames.IsEmpty();
+		if (!bIsPushEnabled)
+		{
+			LabelText = LOCTEXT("PushAssetsToolbarMenuEntryInvalid", "Cannot push. Rundown pages have no valid assets.");
+			TooltipText = LOCTEXT("PushAssetsToolbarMenuEntryTooltipInvalid", "Please ensure the rundown pages are using a valid Avalanche Blueprint");
+		}
+
+		constexpr bool bIsPushing = true;
+		constexpr bool bOpenSubMenuOnClick = false;
+		MenuBuilder.AddSubMenu(
+			LabelText,
+			TooltipText,
+			FNewMenuDelegate::CreateRaw(
+				&StormSyncEditor,
+				&FStormSyncEditorModule::BuildPushAssetsMenuSection,
+				PackageNames,
+				bIsPushing
+			),
+			FUIAction(
+				FExecuteAction(),
+				FCanExecuteAction::CreateLambda([bIsPushEnabled]() { return bIsPushEnabled; })
+			),
+			NAME_None,
+			EUserInterfaceActionType::Button,
+			bOpenSubMenuOnClick,
+			FSlateIcon(FAppStyle::GetAppStyleSetName(), "LevelEditor.ExportAll")
+		);
+
+		MenuBuilder.AddSubMenu(
+			FText::Format(LOCTEXT("CompareAssetsToolbarMenuEntry", "Compare {0} asset(s) with"), FText::AsNumber(PackagesCount)),
+			LOCTEXT("CompareAssetsToolbarMenuEntryTooltip", "Compare asset(s) with a specific remote and see if files (and inner dependencies) are either missing or in mismatched state."),
+			FNewMenuDelegate::CreateRaw(
+				&StormSyncEditor,
+				&FStormSyncEditorModule::BuildCompareWithMenuSection,
+				PackageNames
+			),
+			FUIAction(
+				FExecuteAction(),
+				FCanExecuteAction::CreateLambda([bIsPushEnabled]() { return bIsPushEnabled; })
+			),
+			NAME_None,
+			EUserInterfaceActionType::Button,
+			bOpenSubMenuOnClick,
+			FSlateIcon(FAppStyle::GetAppStyleSetName(), "EditorViewport.Visualizers")
+		);
+	}
+	MenuBuilder.EndSection();
+
+	return MenuBuilder.MakeWidget();
+}
+
+void FStormSyncAvaRundownExtender::PushPackagesToRemote(const FString& RemoteAddressId, const TArray<FName>& InPackageNames)
+{
+	UE_LOG(LogStormSyncAvaBridgeEditor, Display, TEXT("FStormSyncAvaRundownExtender::PushPackagesToRemote - InSelectedPage: (InPackageNames: %d, RemoteAddressId: %s)"), InPackageNames.Num(), *RemoteAddressId);
+
+	FMessageAddress RemoteMessageAddress;
+	if (!FMessageAddress::Parse(RemoteAddressId, RemoteMessageAddress))
+	{
+		UE_LOG(LogStormSyncAvaBridgeEditor, Error, TEXT("Unable to parse %s into a Message Address"), *RemoteAddressId);
+		return;
+	}
+
+	// Note: We sync with a dummy package descriptor, next iterations could add in there an additional UI step.
+	// Something that could be handled with a bit more UI integration, like some kind of popup window or wizard.
+	const FStormSyncPackageDescriptor PackageDescriptor;
+
+	const FOnStormSyncPushComplete Delegate = FOnStormSyncPushComplete::CreateLambda([](const TSharedPtr<FStormSyncTransportPushResponse>& Response)
+	{
+		check(Response.IsValid())
+		UE_LOG(LogStormSyncAvaBridgeEditor, Display, TEXT("FStormSyncAvaRundownExtender::PushPackagesToRemote - Got a response: %s"), *Response->ToString());
+		UStormSyncNotificationSubsystem::Get().HandlePushResponse(Response);
+	});
+
+	IStormSyncTransportClientModule::Get().PushPackages(PackageDescriptor, InPackageNames, RemoteMessageAddress, Delegate);
+}
+
+TArray<FName> FStormSyncAvaRundownExtender::GetSelectedPackagesNames(const UAvaRundown* InRundown, const TWeakPtr<FAvaRundownEditor>& InRundownEditor)
+{
+	check(InRundown);
+
+	TArray<FName> Result;
+
+	const TArray<FAvaRundownPage> SelectedPages = GetSelectedPages(InRundown, InRundownEditor);
+	Algo::Transform(SelectedPages, Result, [InRundown](const FAvaRundownPage& Page)
+	{
+		return Page.GetAssetPath(InRundown).GetLongPackageFName();
+	});
+
+	return Result;
+}
+
+TArray<FAvaRundownPage> FStormSyncAvaRundownExtender::GetSelectedPages(const UAvaRundown* InRundown, const TWeakPtr<FAvaRundownEditor>& InRundownEditor)
+{
+	check(InRundown);
+
+	TArray<FAvaRundownPage> Result;
+	const TSharedPtr<FAvaRundownEditor> RundownEditor = InRundownEditor.Pin();
+	if (!RundownEditor.IsValid())
+	{
+		UE_LOG(LogStormSyncAvaBridgeEditor, Display, TEXT("FStormSyncAvaRundownExtender::GetSelectedPages - Invalid shared ptr from weak ptr delegate param"));
+		return Result;
+	}
+
+	const TConstArrayView<int32> SelectedPageIds = RundownEditor->GetSelectedPagesOnFocusedWidget();
+	for (const int32 SelectedPageId : SelectedPageIds)
+	{
+		FAvaRundownPage AvalanchePage = InRundown->GetPage(SelectedPageId);
+		if (!AvalanchePage.IsValidPage())
+		{
+			continue;
+		}
+
+		FSoftObjectPath SoftAvalancheAssetPath = AvalanchePage.GetAssetPath(InRundown);
+		FString LongPackageName = SoftAvalancheAssetPath.GetLongPackageName();
+		FString AssetName = SoftAvalancheAssetPath.GetAssetName();
+
+		if (!LongPackageName.IsEmpty())
+		{
+			Result.Add(AvalanchePage);
+		}
+	}
+
+	return Result;
+}
+
+TArray<FString> FStormSyncAvaRundownExtender::GetChannelNamesForTemplatePage(const UAvaRundown* InRundown, const FAvaRundownPage& InTemplatePage)
+{
+	check(InRundown);
+
+	TArray<FString> ChannelNames;
+
+	// Here, we try to determine the list of channels to consider for a sync operation, from the selected template page,
+	// with instanced pages that are using the selected package name (Avalanche Blueprint)
+	const FString PackageName = InTemplatePage.GetAssetPath(InRundown).GetLongPackageName();
+
+	// Build the list of instanced pages matching the asset name we want to sync
+	TArray<FAvaRundownPage> Pages = InRundown->GetInstancedPages().Pages.FilterByPredicate([InRundown, PackageName](const FAvaRundownPage& Page)
+	{
+		return Page.GetAssetPath(InRundown).GetLongPackageName() == PackageName;
+	});
+
+	// From there, build a unique list of channel outputs
+	for (const FAvaRundownPage& AvalanchePage : Pages)
+	{
+		ChannelNames.AddUnique(AvalanchePage.GetChannelName().ToString());
+	}
+
+	return ChannelNames;
+}
+
+bool FStormSyncAvaRundownExtender::GetContextMenuSelectionInfos(const UAvaRundown* InRundown, const TWeakPtr<FAvaRundownEditor>& InRundownEditor, bool& bOutIsValidSelection, FText& OutDisabledReasonTooltip, TArray<FName>& OutSelectedPackageNames, FCanExecuteAction& OutCanExecuteAction)
+{
+	if (!FModuleManager::Get().IsModuleLoaded(TEXT("StormSyncEditor")))
+	{
+		UE_LOG(LogStormSyncAvaBridgeEditor, Error, TEXT("FStormSyncAvaRundownExtender::GetContextMenuSelectionInfos - StormSyncEditor module is not loaded. Is StormSync plugin enabled ?"));
+		return false;
+	}
+
+	const FStormSyncEditorModule& StormSyncEditor = FStormSyncEditorModule::Get();
+
+	TArray<FName> SelectedPackagesNames = GetSelectedPackagesNames(InRundown, InRundownEditor);
+	bool bIsValidSelection = !SelectedPackagesNames.IsEmpty();
+
+	// Figure out if selection is containing dirty (unsaved) assets
+	FText DisabledTooltipReason;
+	const TArray<FAssetData> DirtyAssets = StormSyncEditor.GetDirtyAssets(SelectedPackagesNames, DisabledTooltipReason);
+
+	bool bContainsDirtyAssets = !DirtyAssets.IsEmpty();
+	FCanExecuteAction DefaultCanExecuteAction = FCanExecuteAction::CreateLambda([bIsValidSelection, bContainsDirtyAssets]() { return bIsValidSelection && !bContainsDirtyAssets; });
+
+	bOutIsValidSelection = bIsValidSelection;
+	OutDisabledReasonTooltip = DisabledTooltipReason;
+	OutSelectedPackageNames = MoveTemp(SelectedPackagesNames);
+	OutCanExecuteAction = MoveTemp(DefaultCanExecuteAction);
+
+	return true;
+}
+
+#undef LOCTEXT_NAMESPACE
diff --git a/Engine/Plugins/Experimental/StormSyncAvalancheBridge/Source/StormSyncAvaBridgeEditor/Private/StormSyncAvaRundownExtender.h b/Engine/Plugins/Experimental/StormSyncAvalancheBridge/Source/StormSyncAvaBridgeEditor/Private/StormSyncAvaRundownExtender.h
new file mode 100644
--- /dev/null
+++ b/Engine/Plugins/Experimental/StormSyncAvalancheBridge/Source/StormSyncAvaBridgeEditor/Private/StormSyncAvaRundownExtender.h
@@ -0,0 +1,100 @@
+// Copyright Epic Games, Inc. All Rights Reserved.
+
+#pragma once
+
+#include "CoreMinimal.h"
+#include "MessageEndpoint.h"
+#include "Toolkits/AssetEditorToolkit.h"
+
+class FAvaRundownEditor;
+class FExtensibilityManager;
+class UAvaRundown;
+struct FAvaRundownPage;
+
+/** This class handles Motion Design Rundown assets UI extensions for context menu / toolbars */
+class FStormSyncAvaRundownExtender : public TSharedFromThis<FStormSyncAvaRundownExtender>
+{
+public:
+	FStormSyncAvaRundownExtender();
+	~FStormSyncAvaRundownExtender();
+
+private:
+	/** Holds the messaging endpoint. */
+	TSharedPtr<FMessageEndpoint, ESPMode::ThreadSafe> MessageEndpoint;
+
+	/** Name of the extension point we're providing extension next to*/
+	static constexpr const TCHAR* AvalancheExtensionHook = TEXT("PageListOperations");
+	
+	/** Menu extender for motion design rundown editor context menu */
+	FDelegateHandle MenuExtenderHandle;
+	
+	/** Toolbar extender for motion design rundown editor */
+	FDelegateHandle ToolbarExtenderHandle;
+
+	/** Register Context menu / toolbar extensions when editor is fully loaded */
+	void OnPostEngineInit();
+	
+	/** Startup module handler for UI extensions */
+	void RegisterMenuExtensions();
+
+	/** Startup module handler to register an UI extensions from one of Motion Design exposed extensibility hooks */
+	static FDelegateHandle RegisterExtension(const TSharedPtr<FExtensibilityManager> InExtensibilityManager, const FAssetEditorExtender& InExtenderDelegate);
+	
+	/** Shutdown module handler to unregister any UI extensions added here */
+	static void UnregisterExtension(const TSharedPtr<FExtensibilityManager> InExtensibilityManager, const FDelegateHandle& InHandleToRemove);
+	
+	/** Gets the extender to use for rundown context sensitive menus */
+	TSharedRef<FExtender> AddMenuExtender(const TSharedRef<FUICommandList> InCommandList, const TArray<UObject*> ContextSensitiveObjects);
+
+	/* UI Menu Extension handler for template panel */
+	void CreateTemplateContextMenu(FMenuBuilder& MenuBuilder, const UAvaRundown* InRundown, TWeakPtr<FAvaRundownEditor> InRundownEditor);
+
+	/** Context menu handler for initialize action */
+	void HandleInitializeAction(const UAvaRundown* InRundown, TArray<FAvaRundownPage> InSelectedTemplatePages);
+	
+	/** Gets the extender to use for rundown context sensitive menus */
+	TSharedRef<FExtender> AddToolbarExtender(const TSharedRef<FUICommandList> InCommandList, const TArray<UObject*> ContextSensitiveObjects);
+	
+	/** Construct toolbar widgets for rundown sync actions */
+	void FillToolbar(FToolBarBuilder& ToolbarBuilder, TWeakPtr<FAvaRundownEditor> InRundownEditor);
+
+	/** Creates widget for toolbar content */
+	TSharedRef<SWidget> GenerateToolbarMenu(TWeakPtr<FAvaRundownEditor> InRundownEditor);
+
+	/** Storm sync push a list of package names to a remote address ID */
+	static void PushPackagesToRemote(const FString& RemoteAddressId, const TArray<FName>& InPackageNames);
+
+	/** Gather a list of package names from currently selected pages in editor */
+	static TArray<FName> GetSelectedPackagesNames(const UAvaRundown* InRundown, const TWeakPtr<FAvaRundownEditor>& InRundownEditor);
+	
+	/**
+	 * Returns the list of currently selected pages in editor
+	 *
+	 * It will only return selection if the page is referencing a valid Motion Design Asset (not "None" path)
+	 */
+	static TArray<FAvaRundownPage> GetSelectedPages(const UAvaRundown* InRundown, const TWeakPtr<FAvaRundownEditor>& InRundownEditor);
+
+	/** Returns a unique list of channel names, gather from instanced pages, that are matching the passed in template page selection and selected asset name */
+	static TArray<FString> GetChannelNamesForTemplatePage(const UAvaRundown* InRundown, const FAvaRundownPage& InTemplatePage);
+
+	/**
+	 * Provides a way for context menu builders to gather information about current selection.
+	 *
+	 * @param InRundown The rundown UObject we're working with
+	 * @param InRundownEditor The rundown asset editor we're working with
+	 * @param bOutIsValidSelection Output indicating whether selection is valid - eg. current selection returns a list of valid package names
+	 * @param OutDisabledReasonTooltip Output holding the tooltip to use in case the action is disabled - eg. in case the selection contains unsaved dirty assets
+	 * @param OutSelectedPackageNames Output list of package names of the selection, as returned by GetSelectedPackagesNames()
+	 * @param OutCanExecuteAction Output delegate to use as a default CanExecuteAction - eg. disabling the action if selection is invalid
+	 *
+	 * @return false if StormSyncEditor module is not available, and we failed to determine selection state
+	 */
+	static bool GetContextMenuSelectionInfos(
+		const UAvaRundown* InRundown,
+		const TWeakPtr<FAvaRundownEditor>& InRundownEditor,
+		bool& bOutIsValidSelection,
+		FText& OutDisabledReasonTooltip,
+		TArray<FName>& OutSelectedPackageNames,
+		FCanExecuteAction& OutCanExecuteAction
+	);
+};
diff --git a/Engine/Plugins/Experimental/StormSyncAvalancheBridge/Source/StormSyncAvaBridgeEditor/Public/StormSyncAvaBridgeEditor.h b/Engine/Plugins/Experimental/StormSyncAvalancheBridge/Source/StormSyncAvaBridgeEditor/Public/StormSyncAvaBridgeEditor.h
deleted file mode 100644
diff --git a/Engine/Plugins/Experimental/StormSyncAvalancheBridge/Source/StormSyncAvaBridgeEditor/Public/StormSyncAvaBridgeEditorLog.h b/Engine/Plugins/Experimental/StormSyncAvalancheBridge/Source/StormSyncAvaBridgeEditor/Public/StormSyncAvaBridgeEditorLog.h
deleted file mode 100644
