[StormSync] Update user facing text

diff --git a/Engine/Plugins/Experimental/StormSync/Source/StormSyncImport/Private/Subsystems/StormSyncImportSubsystem.cpp b/Engine/Plugins/Experimental/StormSync/Source/StormSyncImport/Private/Subsystems/StormSyncImportSubsystem.cpp
--- a/Engine/Plugins/Experimental/StormSync/Source/StormSyncImport/Private/Subsystems/StormSyncImportSubsystem.cpp
+++ b/Engine/Plugins/Experimental/StormSync/Source/StormSyncImport/Private/Subsystems/StormSyncImportSubsystem.cpp
@@ -415,7 +415,7 @@ void UStormSyncImportSubsystem::HandlePakPostExtract(const FStormSyncPackageDesc
 		LogMessage = FText::Format(LOCTEXT("PakPostExtract_Extract_Failed", "Extracted {0} files. Some content could not be extracted."), FText::AsNumber(FileCount));
 	}
 
-	const FText LogHeading = FText::Format(LOCTEXT("PakPostExtract_Heading", "Received avalanche pak \"{0}\"."), FText::FromString(InPackageDescriptor.Name));
+	const FText LogHeading = FText::Format(LOCTEXT("PakPostExtract_Heading", "Received content spak \"{0}\"."), FText::FromString(InPackageDescriptor.Name));
 	const FText NotifyMessage = FText::Format(LOCTEXT("PakPostExtract_Message_Notification", "{0}\n{1}"), LogHeading, LogMessage);
 	
 	FMessageLog MessageLog(LogName);
diff --git a/Engine/Plugins/Experimental/StormSync/Source/StormSyncTransportClient/Private/Socket/StormSyncTransportClientSocket.h b/Engine/Plugins/Experimental/StormSync/Source/StormSyncTransportClient/Private/Socket/StormSyncTransportClientSocket.h
--- a/Engine/Plugins/Experimental/StormSync/Source/StormSyncTransportClient/Private/Socket/StormSyncTransportClientSocket.h
+++ b/Engine/Plugins/Experimental/StormSync/Source/StormSyncTransportClient/Private/Socket/StormSyncTransportClientSocket.h
@@ -9,7 +9,7 @@
 /**
  * Implements a TCP message tunnel connection.
  *
- * Mainly used to send a raw buffer serialized from StormSyncCoreUtils methods to transfer an avalanche pak on a remote instance.
+ * Mainly used to send a raw buffer serialized from StormSyncCoreUtils methods to transfer a content spak on a remote instance.
  */
 class FStormSyncTransportClientSocket : public FRunnable, public TSharedFromThis<FStormSyncTransportClientSocket>
 {
diff --git a/Engine/Plugins/Experimental/StormSync/Source/StormSyncTransportServer/Private/Socket/StormSyncTransportTcpServer.h b/Engine/Plugins/Experimental/StormSync/Source/StormSyncTransportServer/Private/Socket/StormSyncTransportTcpServer.h
--- a/Engine/Plugins/Experimental/StormSync/Source/StormSyncTransportServer/Private/Socket/StormSyncTransportTcpServer.h
+++ b/Engine/Plugins/Experimental/StormSync/Source/StormSyncTransportServer/Private/Socket/StormSyncTransportTcpServer.h
@@ -52,7 +52,7 @@ struct FClientConnection
  *
  * Heavily based on FSwitchboardListener implementation.
  *
- * Mainly used to receive a raw buffer serialized from StormSyncCoreUtils methods to transfer an avalanche pak.
+ * Mainly used to receive a raw buffer serialized from StormSyncCoreUtils methods to transfer a content spak.
  */
 class FStormSyncTransportTcpServer : public FRunnable
 {
diff --git a/Engine/Plugins/Experimental/StormSyncAvalancheBridge/Source/StormSyncAvaBridgeEditor/Private/StormSyncAvaRundownExtender.cpp b/Engine/Plugins/Experimental/StormSyncAvalancheBridge/Source/StormSyncAvaBridgeEditor/Private/StormSyncAvaRundownExtender.cpp
--- a/Engine/Plugins/Experimental/StormSyncAvalancheBridge/Source/StormSyncAvaBridgeEditor/Private/StormSyncAvaRundownExtender.cpp
+++ b/Engine/Plugins/Experimental/StormSyncAvalancheBridge/Source/StormSyncAvaBridgeEditor/Private/StormSyncAvaRundownExtender.cpp
@@ -149,7 +149,7 @@ void FStormSyncAvaRundownExtender::CreateTemplateContextMenu(FMenuBuilder& MenuB
 	{
 		const FText TooltipText = bIsValidSelection ?
 			FText::Format(LOCTEXT("Initialize_Tooltip", "Sync asset over remote node.{0}"), DisabledTooltipReason) :
-			FText::Format(LOCTEXT("Initialize_Tooltip_Invalid", "Please ensure the rundown page is using a valid Avalanche Blueprint.{0}"), DisabledTooltipReason);
+			FText::Format(LOCTEXT("Initialize_Tooltip_Invalid", "Please ensure the rundown page is using a valid Motion Design Asset.{0}"), DisabledTooltipReason);
 
 		const TArray<FAvaRundownPage> SelectedTemplatePages = GetSelectedPages(InRundown, InRundownEditor);
 
@@ -174,7 +174,7 @@ void FStormSyncAvaRundownExtender::CreateTemplateContextMenu(FMenuBuilder& MenuB
 		if (!bIsValidSelection)
 		{
 			LabelText = LOCTEXT("PushAssetsMenuEntryInvalid", "Cannot push. Page has no valid asset.");
-			TooltipText = LOCTEXT("PushAssetsMenuEntryTooltipInvalid", "Please ensure the rundown pages are using a valid Avalanche Blueprint.");
+			TooltipText = LOCTEXT("PushAssetsMenuEntryTooltipInvalid", "Please ensure the rundown pages are using a valid Motion Design Asset.");
 		}
 		else
 		{
@@ -218,7 +218,7 @@ void FStormSyncAvaRundownExtender::CreateTemplateContextMenu(FMenuBuilder& MenuB
 		if (!bIsValidSelection)
 		{
 			LabelText = LOCTEXT("CompareAssetsMenuEntryInvalid", "Cannot compare. Page has no valid asset.");
-			TooltipText = LOCTEXT("CompareAssetsMenuEntryTooltipInvalid", "Please ensure the rundown pages are using a valid Avalanche Blueprint.");
+			TooltipText = LOCTEXT("CompareAssetsMenuEntryTooltipInvalid", "Please ensure the rundown pages are using a valid Motion Design Asset.");
 		}
 		else
 		{
@@ -412,7 +412,7 @@ TSharedRef<SWidget> FStormSyncAvaRundownExtender::GenerateToolbarMenu(TWeakPtr<F
 		if (!bIsPushEnabled)
 		{
 			LabelText = LOCTEXT("PushAssetsToolbarMenuEntryInvalid", "Cannot push. Rundown pages have no valid assets.");
-			TooltipText = LOCTEXT("PushAssetsToolbarMenuEntryTooltipInvalid", "Please ensure the rundown pages are using a valid Avalanche Blueprint");
+			TooltipText = LOCTEXT("PushAssetsToolbarMenuEntryTooltipInvalid", "Please ensure the rundown pages are using a valid Motion Design Asset");
 		}
 
 		constexpr bool bIsPushing = true;
