[Motion Design]

diff --git a/Engine/Plugins/Experimental/Avalanche/Source/AvalancheViewport/Private/Widgets/SAvaDraggableBoxOverlay.cpp b/Engine/Plugins/Experimental/Avalanche/Source/AvalancheViewport/Private/Widgets/SAvaDraggableBoxOverlay.cpp
--- a/Engine/Plugins/Experimental/Avalanche/Source/AvalancheViewport/Private/Widgets/SAvaDraggableBoxOverlay.cpp
+++ b/Engine/Plugins/Experimental/Avalanche/Source/AvalancheViewport/Private/Widgets/SAvaDraggableBoxOverlay.cpp
@@ -10,7 +10,7 @@
 #include "Widgets/SAvaDraggableBox.h"
 #include "Widgets/SWindow.h"
 
-namespace UE::AvalancheEditor::Private
+namespace UE::AvaEditor::Private
 {
 	static const FAvaShapeEditorViewportControlPosition DefaultPosition = {
 		EHorizontalAlignment::HAlign_Left,
@@ -23,8 +23,8 @@ namespace UE::AvalancheEditor::Private
 
 void SAvaDraggableBoxOverlay::Construct(const FArguments& InArgs)
 {
-	HorizontalAlignment = UE::AvalancheEditor::Private::DefaultPosition.HorizontalAlignment;
-	VerticalAlignment = UE::AvalancheEditor::Private::DefaultPosition.VerticalAlignment;
+	HorizontalAlignment = UE::AvaEditor::Private::DefaultPosition.HorizontalAlignment;
+	VerticalAlignment = UE::AvaEditor::Private::DefaultPosition.VerticalAlignment;
 
 	ChildSlot
 	[
@@ -96,7 +96,7 @@ void SAvaDraggableBoxOverlay::SetBoxAlignmentOffset(const FVector2f& InOffset, b
 		return;
 	}
 
-	using namespace UE::AvalancheEditor::Private;
+	using namespace UE::AvaEditor::Private;
 
 	FVector2f ConstrainedOffset = {
 		FMath::Max(InOffset.X, DraggableBorder),
diff --git a/Engine/Plugins/Experimental/Avalanche/Source/AvalancheViewport/Public/AvaViewportSettings.h b/Engine/Plugins/Experimental/Avalanche/Source/AvalancheViewport/Public/AvaViewportSettings.h
--- a/Engine/Plugins/Experimental/Avalanche/Source/AvalancheViewport/Public/AvaViewportSettings.h
+++ b/Engine/Plugins/Experimental/Avalanche/Source/AvalancheViewport/Public/AvaViewportSettings.h
@@ -29,7 +29,7 @@ UENUM()
 enum class EAvaShapeEditorOverlayType : uint8
 {
 	ComponentVisualizerOnly,
-	FullAvalancheDetails
+	FullDetails
 };
 
 USTRUCT()
diff --git a/Engine/Plugins/Experimental/Avalanche/Source/AvalancheViewport/Public/ViewportClient/IAvaViewportClient.h b/Engine/Plugins/Experimental/Avalanche/Source/AvalancheViewport/Public/ViewportClient/IAvaViewportClient.h
--- a/Engine/Plugins/Experimental/Avalanche/Source/AvalancheViewport/Public/ViewportClient/IAvaViewportClient.h
+++ b/Engine/Plugins/Experimental/Avalanche/Source/AvalancheViewport/Public/ViewportClient/IAvaViewportClient.h
@@ -30,7 +30,7 @@ class IAvaViewportClient : public IAvaViewportWorldCoordinateConverter
 public:
 	UE_AVA_INHERITS(IAvaViewportClient, IAvaViewportWorldCoordinateConverter)
 
-	virtual bool IsAvalancheViewport() const = 0;
+	virtual bool IsMotionDesignViewport() const = 0;
 
 	virtual const FEditorViewportClient* AsEditorViewportClient() const = 0;
 
diff --git a/Engine/Plugins/Experimental/StormSync/Source/StormSyncImport/Private/Subsystems/StormSyncImportSubsystem.h b/Engine/Plugins/Experimental/StormSync/Source/StormSyncImport/Private/Subsystems/StormSyncImportSubsystem.h
--- a/Engine/Plugins/Experimental/StormSync/Source/StormSyncImport/Private/Subsystems/StormSyncImportSubsystem.h
+++ b/Engine/Plugins/Experimental/StormSync/Source/StormSyncImport/Private/Subsystems/StormSyncImportSubsystem.h
@@ -148,7 +148,7 @@ private:
 	/**
 	 * Our own custom version of ObjectTools::DeleteAssets
 	 *
-	 * We want to force delete the file, to allow assets with on disk references (for things like Textures in an Avalanche Blueprints) to be properly overwritten,
+	 * We want to force delete the file, to allow assets with on disk references (for things like Textures in a Motion Design asset) to be properly overwritten,
 	 * if referenced in memory, opened in editor etc.
 	 *
 	 * ObjectTools::DeleteAssets will only succeed if the AssetDeleteModel doesn't have on disk references (see CanDelete / CanForceDelete)
diff --git a/Engine/Plugins/Experimental/StormSyncAvalancheBridge/Source/StormSyncAvaBridge/Private/ModularFeature/StormSyncAvaSyncProvider.h b/Engine/Plugins/Experimental/StormSyncAvalancheBridge/Source/StormSyncAvaBridge/Private/ModularFeature/StormSyncAvaSyncProvider.h
--- a/Engine/Plugins/Experimental/StormSyncAvalancheBridge/Source/StormSyncAvaBridge/Private/ModularFeature/StormSyncAvaSyncProvider.h
+++ b/Engine/Plugins/Experimental/StormSyncAvalancheBridge/Source/StormSyncAvaBridge/Private/ModularFeature/StormSyncAvaSyncProvider.h
@@ -13,7 +13,7 @@ struct FStormSyncTransportStatusResponse;
 struct FStormSyncTransportSyncResponse;
 
 /**
- * Avalanche synchronization provider feature implementation
+ * Motion Design synchronization provider feature implementation
  */
 class FStormSyncAvaSyncProvider : public IAvaMediaSyncProvider
 {
diff --git a/Engine/Plugins/Experimental/StormSyncAvalancheBridge/Source/StormSyncAvaBridge/Private/StormSyncAvaBridge.cpp b/Engine/Plugins/Experimental/StormSyncAvalancheBridge/Source/StormSyncAvaBridge/Private/StormSyncAvaBridge.cpp
--- a/Engine/Plugins/Experimental/StormSyncAvalancheBridge/Source/StormSyncAvaBridge/Private/StormSyncAvaBridge.cpp
+++ b/Engine/Plugins/Experimental/StormSyncAvalancheBridge/Source/StormSyncAvaBridge/Private/StormSyncAvaBridge.cpp
@@ -157,7 +157,7 @@ void FStormSyncAvaBridgeModule::OnStormSyncServerStopped()
 
 void FStormSyncAvaBridgeModule::RegisterUserDataForPlaybackServer()
 {
-	// Fail safe checks for both Avalanche and Storm Sync modules, since this method can be executed from delegates in either of these modules
+	// Fail safe checks for both Motion Design and Storm Sync modules, since this method can be executed from delegates in either of these modules
 	if (!ValidateModulesAreAvailable())
 	{
 		return;
@@ -181,7 +181,7 @@ void FStormSyncAvaBridgeModule::RegisterUserDataForPlaybackServer()
 
 void FStormSyncAvaBridgeModule::RegisterUserDataForPlaybackClient()
 {
-	// Fail safe checks for both Avalanche and Storm Sync modules, since this method can be executed from delegates in either of these modules
+	// Fail safe checks for both Motion Design and Storm Sync modules, since this method can be executed from delegates in either of these modules
 	if (!ValidateModulesAreAvailable())
 	{
 		return;
diff --git a/Engine/Plugins/Experimental/StormSyncAvalancheBridge/Source/StormSyncAvaBridgeEditor/Private/StormSyncAvaRundownExtender.cpp b/Engine/Plugins/Experimental/StormSyncAvalancheBridge/Source/StormSyncAvaBridgeEditor/Private/StormSyncAvaRundownExtender.cpp
--- a/Engine/Plugins/Experimental/StormSyncAvalancheBridge/Source/StormSyncAvaBridgeEditor/Private/StormSyncAvaRundownExtender.cpp
+++ b/Engine/Plugins/Experimental/StormSyncAvalancheBridge/Source/StormSyncAvaBridgeEditor/Private/StormSyncAvaRundownExtender.cpp
@@ -107,7 +107,7 @@ TSharedRef<FExtender> FStormSyncAvaRundownExtender::AddMenuExtender(const TShare
 
 	// Template panel extension
 	Extender->AddMenuExtension(
-		AvalancheExtensionHook,
+		MenuExtensionHook,
 		EExtensionHook::After,
 		InCommandList,
 		// Convert to weak ptr to prevent ownership to the rundown editor and potentially increasing its lifetime
@@ -375,11 +375,11 @@ TSharedRef<SWidget> FStormSyncAvaRundownExtender::GenerateToolbarMenu(TWeakPtr<F
 	{
 		const UAvaRundown* Rundown = RundownEditor->GetRundown(); 
 
-		// Gather the list of all package names from Avalanche Blueprints in this Rundown pages
-		const FAvaRundownPageCollection& AvalanchePageCollection = Rundown->GetTemplatePages();
-		for (const FAvaRundownPage& AvalanchePage : AvalanchePageCollection.Pages)
+		// Gather the list of all package names from Motion Design assets in this Rundown pages
+		const FAvaRundownPageCollection& PageCollection = Rundown->GetTemplatePages();
+		for (const FAvaRundownPage& Page : PageCollection.Pages)
 		{
-			if (FString PackageName = AvalanchePage.GetAssetPath(Rundown).GetLongPackageName(); !PackageName.IsEmpty())
+			if (FString PackageName = Page.GetAssetPath(Rundown).GetLongPackageName(); !PackageName.IsEmpty())
 			{
 				PackageNames.AddUnique(FName(*PackageName));
 			}
@@ -514,19 +514,19 @@ TArray<FAvaRundownPage> FStormSyncAvaRundownExtender::GetSelectedPages(const UAv
 	const TConstArrayView<int32> SelectedPageIds = RundownEditor->GetSelectedPagesOnFocusedWidget();
 	for (const int32 SelectedPageId : SelectedPageIds)
 	{
-		FAvaRundownPage AvalanchePage = InRundown->GetPage(SelectedPageId);
-		if (!AvalanchePage.IsValidPage())
+		FAvaRundownPage Page = InRundown->GetPage(SelectedPageId);
+		if (!Page.IsValidPage())
 		{
 			continue;
 		}
 
-		FSoftObjectPath SoftAvalancheAssetPath = AvalanchePage.GetAssetPath(InRundown);
-		FString LongPackageName = SoftAvalancheAssetPath.GetLongPackageName();
-		FString AssetName = SoftAvalancheAssetPath.GetAssetName();
+		FSoftObjectPath SoftAssetPath = Page.GetAssetPath(InRundown);
+		FString LongPackageName = SoftAssetPath.GetLongPackageName();
+		FString AssetName = SoftAssetPath.GetAssetName();
 
 		if (!LongPackageName.IsEmpty())
 		{
-			Result.Add(AvalanchePage);
+			Result.Add(Page);
 		}
 	}
 
@@ -540,7 +540,7 @@ TArray<FString> FStormSyncAvaRundownExtender::GetChannelNamesForTemplatePage(con
 	TArray<FString> ChannelNames;
 
 	// Here, we try to determine the list of channels to consider for a sync operation, from the selected template page,
-	// with instanced pages that are using the selected package name (Avalanche Blueprint)
+	// with instanced pages that are using the selected package name (Motion Design asset)
 	const FString PackageName = InTemplatePage.GetAssetPath(InRundown).GetLongPackageName();
 
 	// Build the list of instanced pages matching the asset name we want to sync
@@ -550,9 +550,9 @@ TArray<FString> FStormSyncAvaRundownExtender::GetChannelNamesForTemplatePage(con
 	});
 
 	// From there, build a unique list of channel outputs
-	for (const FAvaRundownPage& AvalanchePage : Pages)
+	for (const FAvaRundownPage& Page : Pages)
 	{
-		ChannelNames.AddUnique(AvalanchePage.GetChannelName().ToString());
+		ChannelNames.AddUnique(Page.GetChannelName().ToString());
 	}
 
 	return ChannelNames;
diff --git a/Engine/Plugins/Experimental/StormSyncAvalancheBridge/Source/StormSyncAvaBridgeEditor/Private/StormSyncAvaRundownExtender.h b/Engine/Plugins/Experimental/StormSyncAvalancheBridge/Source/StormSyncAvaBridgeEditor/Private/StormSyncAvaRundownExtender.h
--- a/Engine/Plugins/Experimental/StormSyncAvalancheBridge/Source/StormSyncAvaBridgeEditor/Private/StormSyncAvaRundownExtender.h
+++ b/Engine/Plugins/Experimental/StormSyncAvalancheBridge/Source/StormSyncAvaBridgeEditor/Private/StormSyncAvaRundownExtender.h
@@ -22,8 +22,8 @@ private:
 	/** Holds the messaging endpoint. */
 	TSharedPtr<FMessageEndpoint, ESPMode::ThreadSafe> MessageEndpoint;
 
-	/** Name of the extension point we're providing extension next to*/
-	static constexpr const TCHAR* AvalancheExtensionHook = TEXT("PageListOperations");
+	/** Name of the extension point we're providing extension next to */
+	static constexpr const TCHAR* MenuExtensionHook = TEXT("PageListOperations");
 	
 	/** Menu extender for motion design rundown editor context menu */
 	FDelegateHandle MenuExtenderHandle;
